- Whitelist IOS for EOSShared

#horde 96357

diff --git a/Engine/Plugins/Online/EOSShared/EOSShared.uplugin b/Engine/Plugins/Online/EOSShared/EOSShared.uplugin
--- a/Engine/Plugins/Online/EOSShared/EOSShared.uplugin
+++ b/Engine/Plugins/Online/EOSShared/EOSShared.uplugin
@@ -17,6 +17,7 @@
 			"LoadingPhase": "PostConfigInit",
 			"WhitelistPlatforms": [
 				"Android",
+				"IOS",
 				"Linux",
 				"Mac",
 				"Win64"
