- Merge fixes for consoles not loading EOSSDK from Release-Engine-Test

diff --git a/Engine/Plugins/Online/EOSShared/Source/EOSShared/Private/EOSSDKManager.cpp b/Engine/Plugins/Online/EOSShared/Source/EOSShared/Private/EOSSDKManager.cpp
--- a/Engine/Plugins/Online/EOSShared/Source/EOSShared/Private/EOSSDKManager.cpp
+++ b/Engine/Plugins/Online/EOSShared/Source/EOSShared/Private/EOSSDKManager.cpp
@@ -98,6 +98,11 @@ namespace
 			Result = FPlatformProcess::GetDllHandle(*EngineBinaryPath);
 		}
 
+		if (!Result)
+		{
+			Result = FPlatformProcess::GetDllHandle(TEXT(EOSSDK_RUNTIME_LIBRARY_NAME));
+		}
+
 		return Result;
 	}
 #endif
