- Fix EOSSDK not being loaded from Engine folder in UE5 when searched for by name, needs an explicit path.

diff --git a/Engine/Plugins/Online/EOSShared/Source/EOSShared/Private/EOSSDKManager.cpp b/Engine/Plugins/Online/EOSShared/Source/EOSShared/Private/EOSSDKManager.cpp
--- a/Engine/Plugins/Online/EOSShared/Source/EOSShared/Private/EOSSDKManager.cpp
+++ b/Engine/Plugins/Online/EOSShared/Source/EOSShared/Private/EOSSDKManager.cpp
@@ -89,14 +89,14 @@ namespace
 	{
 		void* Result = nullptr;
 
-		const FString SDKBinaryPath = FPaths::Combine(FPaths::ProjectDir(), TEXT("Binaries"), FPlatformProcess::GetBinariesSubdirectory(), TEXT(EOSSDK_RUNTIME_LIBRARY_NAME));
-		Result = FPlatformProcess::GetDllHandle(*SDKBinaryPath);
+		const FString ProjectBinaryPath = FPaths::Combine(FPaths::ProjectDir(), TEXT("Binaries"), FPlatformProcess::GetBinariesSubdirectory(), TEXT(EOSSDK_RUNTIME_LIBRARY_NAME));
+		Result = FPlatformProcess::GetDllHandle(*ProjectBinaryPath);
 
 		if (!Result)
 		{
-			// Fallback on searching by name
-			Result = FPlatformProcess::GetDllHandle(TEXT(EOSSDK_RUNTIME_LIBRARY_NAME));
-}
+			const FString EngineBinaryPath = FPaths::Combine(FPaths::EngineDir(), TEXT("Binaries"), FPlatformProcess::GetBinariesSubdirectory(), TEXT(EOSSDK_RUNTIME_LIBRARY_NAME));
+			Result = FPlatformProcess::GetDllHandle(*EngineBinaryPath);
+		}
 
 		return Result;
 	}
