Changed the loading phase for the EOSShared plugin

diff --git a/Engine/Plugins/Online/EOSShared/EOSShared.uplugin b/Engine/Plugins/Online/EOSShared/EOSShared.uplugin
--- a/Engine/Plugins/Online/EOSShared/EOSShared.uplugin
+++ b/Engine/Plugins/Online/EOSShared/EOSShared.uplugin
@@ -14,6 +14,7 @@
 		{
 			"Name": "EOSShared",
 			"Type": "Runtime",
+			"LoadingPhase": "PostConfigInit",
 			"WhitelistPlatforms": [
 				"Android",
 				"IOS",
