- Fix CIS errors in non-unity mode

diff --git a/Engine/Plugins/Online/EOSShared/Source/EOSShared/Private/EOSSharedModule.cpp b/Engine/Plugins/Online/EOSShared/Source/EOSShared/Private/EOSSharedModule.cpp
--- a/Engine/Plugins/Online/EOSShared/Source/EOSShared/Private/EOSSharedModule.cpp
+++ b/Engine/Plugins/Online/EOSShared/Source/EOSShared/Private/EOSSharedModule.cpp
@@ -5,6 +5,7 @@
 #include "Features/IModularFeatures.h"
 #include "Misc/CoreDelegates.h"
 #include "Modules/ModuleManager.h"
+#include "EOSSharedSettings.h"
 
 #if WITH_EDITOR
 #include "ISettingsModule.h"
diff --git a/Engine/Plugins/Online/EOSShared/Source/EOSShared/Private/EOSSharedSettings.cpp b/Engine/Plugins/Online/EOSShared/Source/EOSShared/Private/EOSSharedSettings.cpp
--- a/Engine/Plugins/Online/EOSShared/Source/EOSShared/Private/EOSSharedSettings.cpp
+++ b/Engine/Plugins/Online/EOSShared/Source/EOSShared/Private/EOSSharedSettings.cpp
@@ -1,15 +1,14 @@
 // Copyright Epic Games, Inc. All Rights Reserved.
 
 #include "EOSSharedSettings.h"
-#include "EOSShared.h"
 
 #include "Algo/Transform.h"
 #include "Misc/CommandLine.h"
 #include "Misc/ConfigCacheIni.h"
+#include "EOSShared.h"
+#include "IEOSSDKManager.h"
 
-#if WITH_EDITOR
-#include "Misc/MessageDialog.h"
-#endif
+#include "eos_logging.h"
 
 #define LOCTEXT_NAMESPACE "EOS"
 #define INI_SECTION TEXT("/Script/EOSShared.EOSSharedSettings")
