- Fix CIS issues with UE4Game Win32 etc Shipping builds depending on the Settings module which is apparently not valid. We only need it in Editor builds so limit it to that.

diff --git a/Engine/Plugins/Online/EOSShared/Source/EOSShared/EOSShared.Build.cs b/Engine/Plugins/Online/EOSShared/Source/EOSShared/EOSShared.Build.cs
--- a/Engine/Plugins/Online/EOSShared/Source/EOSShared/EOSShared.Build.cs
+++ b/Engine/Plugins/Online/EOSShared/Source/EOSShared/EOSShared.Build.cs
@@ -15,9 +15,18 @@ public class EOSShared : ModuleRules
 			{
 				"Core",
 				"CoreUObject",
-				"EOSSDK",
-				"Settings"
+				"EOSSDK"
 			}
 		);
+
+		if (Target.bBuildEditor)
+		{
+			PrivateDependencyModuleNames.AddRange(
+				new string[]
+				{
+					"Settings"
+				}
+			);
+		}
 	}
 }
