Changed ApplicationName and ApplicationVersion passed to eos init.

diff --git a/Engine/Plugins/Online/EOSShared/Source/EOSShared/Private/EOSSDKManager.cpp b/Engine/Plugins/Online/EOSShared/Source/EOSShared/Private/EOSSDKManager.cpp
--- a/Engine/Plugins/Online/EOSShared/Source/EOSShared/Private/EOSSDKManager.cpp
+++ b/Engine/Plugins/Online/EOSShared/Source/EOSShared/Private/EOSSDKManager.cpp
@@ -11,6 +11,7 @@
 #include "Misc/CoreDelegates.h"
 #include "Misc/NetworkVersion.h"
 #include "Misc/Paths.h"
+#include "Misc/ConfigCacheIni.h"
 #include "ProfilingDebugging/CsvProfiler.h"
 #include "Stats/Stats.h"
 
@@ -283,13 +284,17 @@ void FEOSSDKManager::OnLogVerbosityChanged(const FLogCategoryName& CategoryName,
 
 FString FEOSSDKManager::GetProductName() const
 {
-	return FApp::GetProjectName();
+	FString ProductName;
+	if (!GConfig->GetString(TEXT("EOSSDK"), TEXT("ProductName"), ProductName, GEngineIni))
+	{
+		ProductName = FApp::GetProjectName();
+	}
+	return ProductName;
 }
 
 FString FEOSSDKManager::GetProductVersion() const
 {
-	static const FString ProductVersion = FNetworkVersion::GetProjectVersion().IsEmpty() ? TEXT("Unknown") : FNetworkVersion::GetProjectVersion();
-	return ProductVersion;
+	return FApp::GetBuildVersion();
 }
 
 void FEOSSDKManager::ReleasePlatform(EOS_HPlatform PlatformHandle)
