Interpret EOS logs as utf8 strings rather than ansi

1) From api docs, LogMessage category and message are utf8 lines.

Test plan:
1) Run a Win64 app with EOS non-ansi symbol logs. It logs the same lines.

diff --git a/Engine/Plugins/Online/EOSShared/Source/EOSShared/Private/EOSSDKManager.cpp b/Engine/Plugins/Online/EOSShared/Source/EOSShared/Private/EOSSDKManager.cpp
--- a/Engine/Plugins/Online/EOSShared/Source/EOSShared/Private/EOSSDKManager.cpp
+++ b/Engine/Plugins/Online/EOSShared/Source/EOSShared/Private/EOSSDKManager.cpp
@@ -47,9 +47,9 @@ namespace
 #if !NO_LOGGING
 	void EOS_CALL EOSLogMessageReceived(const EOS_LogMessage* Message)
 	{
-#define EOSLOG(Level) UE_LOG(LogEOSSDK, Level, TEXT("%s: %s"), ANSI_TO_TCHAR(Message->Category), *MessageStr)
+#define EOSLOG(Level) UE_LOG(LogEOSSDK, Level, TEXT("%s: %s"), UTF8_TO_TCHAR(Message->Category), *MessageStr)
 
-		FString MessageStr(ANSI_TO_TCHAR(Message->Message));
+		FString MessageStr(UTF8_TO_TCHAR(Message->Message));
 		MessageStr.TrimStartAndEndInline();
 
 		switch (Message->Level)
