- Add FCoreDelegates::OnConfigSectionChanged delegate
- Fire FCoreDelegates::OnConfigSectionChanged when hotfix is applied.
- Add configurable tick interval to limit EOS SDK platform update rate. One tick per platform will occur in each interval.

diff --git a/Engine/Plugins/Online/EOSShared/Source/EOSShared/Private/EOSSDKManager.cpp b/Engine/Plugins/Online/EOSShared/Source/EOSShared/Private/EOSSDKManager.cpp
--- a/Engine/Plugins/Online/EOSShared/Source/EOSShared/Private/EOSSDKManager.cpp
+++ b/Engine/Plugins/Online/EOSShared/Source/EOSShared/Private/EOSSDKManager.cpp
@@ -194,6 +194,9 @@ EOS_EResult FEOSSDKManager::Initialize()
 		{
 			bInitialized = true;
 
+			FCoreDelegates::OnConfigSectionChanged.AddRaw(this, &FEOSSDKManager::OnConfigSectionChanged);
+			LoadConfig();
+
 #if !NO_LOGGING
 			FCoreDelegates::OnLogVerbosityChanged.AddRaw(this, &FEOSSDKManager::OnLogVerbosityChanged);
 
@@ -232,11 +235,7 @@ IEOSPlatformHandlePtr FEOSSDKManager::CreatePlatform(EOS_Platform_Options& Platf
 		{
 			ActivePlatforms.Emplace(PlatformHandle);
 			SharedPlatform = MakeShared<FEOSPlatformHandle, ESPMode::ThreadSafe>(*this, PlatformHandle);
-
-			if (!TickerHandle.IsValid())
-			{
-				TickerHandle = FTSTicker::GetCoreTicker().AddTicker(FTickerDelegate::CreateRaw(this, &FEOSSDKManager::Tick), 0.0f);
-			}
+			SetupTicker();
 		}
 		else
 		{
@@ -251,20 +250,64 @@ IEOSPlatformHandlePtr FEOSSDKManager::CreatePlatform(EOS_Platform_Options& Platf
 	return SharedPlatform;
 }
 
+void FEOSSDKManager::OnConfigSectionChanged(const TCHAR* IniFilename, const TCHAR* SectionName)
+{
+	if (!FCString::Strcmp(IniFilename, *GEngineIni) && !FCString::Strcmp(SectionName, TEXT("EOSSDK")))
+	{
+		LoadConfig();
+	}
+}
+
+void FEOSSDKManager::LoadConfig()
+{
+	ConfigTickIntervalSeconds = 0.f;
+	GConfig->GetDouble(TEXT("EOSSDK"), TEXT("TickIntervalSeconds"), ConfigTickIntervalSeconds, GEngineIni);
+
+	SetupTicker();
+}
+
+void FEOSSDKManager::SetupTicker()
+{
+	if (TickerHandle.IsValid())
+	{
+		FTSTicker::GetCoreTicker().RemoveTicker(TickerHandle);
+		TickerHandle.Reset();
+	}
+
+	if (ActivePlatforms.Num() > 0)
+	{
+		const double TickIntervalSeconds = ConfigTickIntervalSeconds > SMALL_NUMBER ? ConfigTickIntervalSeconds / ActivePlatforms.Num() : 0.f;
+		TickerHandle = FTSTicker::GetCoreTicker().AddTicker(FTickerDelegate::CreateRaw(this, &FEOSSDKManager::Tick), TickIntervalSeconds);
+	}
+}
+
 bool FEOSSDKManager::Tick(float)
 {
 	ReleaseReleasedPlatforms();
 
-	//LLM_SCOPE(ELLMTag::EOSSDK); // TODO
-	for (EOS_HPlatform PlatformHandle : ActivePlatforms)
+	if (ActivePlatforms.Num())
 	{
-		LLM_SCOPE(ELLMTag::RealTimeCommunications);
-		QUICK_SCOPE_CYCLE_COUNTER(FEOSSDKManager_Tick);
-		CSV_SCOPED_TIMING_STAT_EXCLUSIVE(EOSSDK);
-		EOS_Platform_Tick(PlatformHandle);
+		TArray<EOS_HPlatform> PlatformsToTick;
+		if (ConfigTickIntervalSeconds > SMALL_NUMBER)
+		{
+			PlatformTickIdx = (PlatformTickIdx + 1) % ActivePlatforms.Num();
+			PlatformsToTick.Emplace(ActivePlatforms[PlatformTickIdx]);
+		}
+		else
+		{
+			PlatformsToTick = ActivePlatforms;
+		}
+
+		for (EOS_HPlatform PlatformHandle : PlatformsToTick)
+		{
+			LLM_SCOPE(ELLMTag::RealTimeCommunications); // TODO should really be ELLMTag::EOSSDK
+			QUICK_SCOPE_CYCLE_COUNTER(FEOSSDKManager_Tick);
+			CSV_SCOPED_TIMING_STAT_EXCLUSIVE(EOSSDK);
+			EOS_Platform_Tick(PlatformHandle);
+		}
 	}
 
-	return ActivePlatforms.Num() > 0;
+	return true;
 }
 
 void FEOSSDKManager::OnLogVerbosityChanged(const FLogCategoryName& CategoryName, ELogVerbosity::Type OldVerbosity, ELogVerbosity::Type NewVerbosity)
@@ -308,23 +351,20 @@ void FEOSSDKManager::ReleasePlatform(EOS_HPlatform PlatformHandle)
 
 void FEOSSDKManager::ReleaseReleasedPlatforms()
 {
-	for (EOS_HPlatform PlatformHandle : ReleasedPlatforms)
+	if (ReleasedPlatforms.Num() > 0)
 	{
-		if (ensure(ActivePlatforms.Contains(PlatformHandle)))
+		for (EOS_HPlatform PlatformHandle : ReleasedPlatforms)
 		{
-		EOS_Platform_Release(PlatformHandle);
-			ActivePlatforms.Remove(PlatformHandle);
-		}
-	}
-	ReleasedPlatforms.Empty();
-
-	if (TickerHandle.IsValid() &&
-		ActivePlatforms.Num() == 0)
-		{
-			FTSTicker::GetCoreTicker().RemoveTicker(TickerHandle);
-			TickerHandle.Reset();
+			if (ensure(ActivePlatforms.Contains(PlatformHandle)))
+			{
+				EOS_Platform_Release(PlatformHandle);
+				ActivePlatforms.Remove(PlatformHandle);
+			}
 		}
+		ReleasedPlatforms.Empty();
+		SetupTicker();
 	}
+}
 
 void FEOSSDKManager::Shutdown()
 {
@@ -334,11 +374,13 @@ void FEOSSDKManager::Shutdown()
 		ReleaseReleasedPlatforms();
 
 		if (ActivePlatforms.Num() > 0)
-			{
+		{
 			UE_LOG(LogEOSSDK, Warning, TEXT("FEOSSDKManager::Shutdown Releasing %d remaining platforms"), ActivePlatforms.Num());
 			ReleasedPlatforms.Append(ActivePlatforms);
 			ReleaseReleasedPlatforms();
-			}
+		}
+
+		FCoreDelegates::OnConfigSectionChanged.RemoveAll(this);
 
 #if !NO_LOGGING
 		FCoreDelegates::OnLogVerbosityChanged.RemoveAll(this);
diff --git a/Engine/Plugins/Online/EOSShared/Source/EOSShared/Private/EOSSDKManager.h b/Engine/Plugins/Online/EOSShared/Source/EOSShared/Private/EOSSDKManager.h
--- a/Engine/Plugins/Online/EOSShared/Source/EOSShared/Private/EOSSDKManager.h
+++ b/Engine/Plugins/Online/EOSShared/Source/EOSShared/Private/EOSSDKManager.h
@@ -4,7 +4,9 @@
 
 #if WITH_EOS_SDK
 
+#include "Containers/Ticker.h"
 #include "Containers/UnrealString.h"
+#include "IEOSSDKManager.h"
 
 #if defined(EOS_PLATFORM_BASE_FILE_NAME)
 #include EOS_PLATFORM_BASE_FILE_NAME
@@ -12,9 +14,6 @@
 #include "eos_common.h"
 #include "eos_init.h"
 
-#include "IEOSSDKManager.h"
-#include "Containers/Ticker.h"
-
 struct FEOSPlatformHandle;
 
 class FEOSSDKManager : public IEOSSDKManager
@@ -41,8 +40,11 @@ protected:
 private:
 	friend struct FEOSPlatformHandle;
 
+	void OnConfigSectionChanged(const TCHAR* IniFilename, const TCHAR* SectionName);
+	void LoadConfig();
 	void ReleasePlatform(EOS_HPlatform PlatformHandle);
 	void ReleaseReleasedPlatforms();
+	void SetupTicker();
 	bool Tick(float);
 	void OnLogVerbosityChanged(const FLogCategoryName& CategoryName, ELogVerbosity::Type OldVerbosity, ELogVerbosity::Type NewVerbosity);
 
@@ -52,12 +54,18 @@ private:
 
 	/** Are we currently initialized */
 	bool bInitialized = false;
+	/** Index of the last ticked platform, used for round-robin ticking when ConfigTickIntervalSeconds > 0 */
+	uint8 PlatformTickIdx = 0;
 	/** Created platforms actively ticking */
 	TArray<EOS_HPlatform> ActivePlatforms;
 	/** Contains platforms released with ReleasePlatform, which we will release on the next Tick. */
 	TArray<EOS_HPlatform> ReleasedPlatforms;
 	/** Handle to ticker delegate for Tick(), valid whenever there are ActivePlatforms to tick, or ReleasedPlatforms to release. */
 	FTSTicker::FDelegateHandle TickerHandle;
+
+	// Config
+	/** Interval between platform ticks. 0 means we tick every frame. */
+	double ConfigTickIntervalSeconds = 0.f;
 };
 
 struct FEOSPlatformHandle : public IEOSPlatformHandle
diff --git a/Engine/Plugins/Online/OnlineFramework/Source/Hotfix/Private/OnlineHotfixManager.cpp b/Engine/Plugins/Online/OnlineFramework/Source/Hotfix/Private/OnlineHotfixManager.cpp
--- a/Engine/Plugins/Online/OnlineFramework/Source/Hotfix/Private/OnlineHotfixManager.cpp
+++ b/Engine/Plugins/Online/OnlineFramework/Source/Hotfix/Private/OnlineHotfixManager.cpp
@@ -922,6 +922,7 @@ bool UOnlineHotfixManager::HotfixIniFile(const FString& FileName, const FString&
 	bool bUpdateHttpConfigs = false;
 	bool bUpdateOnlineTracing = false;
 	TSet<FString> OnlineSubSections;
+	TSet<FString> UpdatedSectionNames;
 	// Find the set of object classes that were affected
 	while (StartIndex >= 0 && StartIndex < IniData.Len() && EndIndex >= StartIndex)
 	{
@@ -963,6 +964,8 @@ bool UOnlineHotfixManager::HotfixIniFile(const FString& FileName, const FString&
 					}
 				}
 
+				UpdatedSectionNames.Emplace(IniData.Mid(StartIndex+1, EndIndex - StartIndex - 1));
+
 				int32 PerObjectNameIndex = IniData.Find(TEXT(" "), ESearchCase::IgnoreCase, ESearchDir::FromStart, StartIndex);
 
 				const TCHAR* AssetHotfixIniHACK = TEXT("[AssetHotfix]");
@@ -974,6 +977,7 @@ bool UOnlineHotfixManager::HotfixIniFile(const FString& FileName, const FString&
 
 				if (bIsEngineIni)
 				{
+					// TODO replace all of this with bindees to FCoreDelegates::OnConfigSectionChanged
 					const TCHAR* LogConfigSection = TEXT("[Core.Log]");
 					const TCHAR* ConsoleVariableSection = TEXT("[ConsoleVariables]");
 					const TCHAR* HttpSection = TEXT("[HTTP"); // note "]" omitted on purpose since we want a partial match
@@ -997,8 +1001,8 @@ bool UOnlineHotfixManager::HotfixIniFile(const FString& FileName, const FString&
 					}
 					else if (FCString::Strnicmp(*IniData + StartIndex, OnlineSubSectionKey, FCString::Strlen(OnlineSubSectionKey)) == 0)
 					{
-						const FString SectionStr = IniData.Mid(StartIndex, EndIndex - StartIndex + 1);
-						OnlineSubSections.Add(SectionStr);
+						FString SectionStr = IniData.Mid(StartIndex, EndIndex - StartIndex + 1);
+						OnlineSubSections.Emplace(MoveTemp(SectionStr));
 					}
 				}
 
@@ -1097,6 +1101,12 @@ bool UOnlineHotfixManager::HotfixIniFile(const FString& FileName, const FString&
 		NumObjectsReloaded++;
 	}
 
+	const FString ConfigFileName = ConfigFile->Name.ToString();
+	for (const FString& SectionName : UpdatedSectionNames)
+	{
+		FCoreDelegates::OnConfigSectionChanged.Broadcast(*ConfigFileName, *SectionName);
+	}
+
 	// Reload log suppression if configs changed
 	if (bUpdateLogSuppression)
 	{
diff --git a/Engine/Source/Runtime/Core/Private/Misc/CoreDelegates.cpp b/Engine/Source/Runtime/Core/Private/Misc/CoreDelegates.cpp
--- a/Engine/Source/Runtime/Core/Private/Misc/CoreDelegates.cpp
+++ b/Engine/Source/Runtime/Core/Private/Misc/CoreDelegates.cpp
@@ -142,6 +142,7 @@ FCoreDelegates::FOnFConfigFileCreated FCoreDelegates::OnFConfigDeleted;
 FCoreDelegates::FOnConfigValueRead FCoreDelegates::OnConfigValueRead;
 FCoreDelegates::FOnConfigSectionRead FCoreDelegates::OnConfigSectionRead;
 FCoreDelegates::FOnConfigSectionRead FCoreDelegates::OnConfigSectionNameRead;
+FCoreDelegates::FOnConfigSectionChanged FCoreDelegates::OnConfigSectionChanged;
 FCoreDelegates::FOnApplyCVarFromIni FCoreDelegates::OnApplyCVarFromIni;
 FCoreDelegates::FOnSystemResolutionChanged FCoreDelegates::OnSystemResolutionChanged;
 
diff --git a/Engine/Source/Runtime/Core/Public/Misc/CoreDelegates.h b/Engine/Source/Runtime/Core/Public/Misc/CoreDelegates.h
--- a/Engine/Source/Runtime/Core/Public/Misc/CoreDelegates.h
+++ b/Engine/Source/Runtime/Core/Public/Misc/CoreDelegates.h
@@ -342,6 +342,9 @@ public:
 	static FOnConfigSectionRead OnConfigSectionRead;
 	static FOnConfigSectionRead OnConfigSectionNameRead;
 
+	DECLARE_MULTICAST_DELEGATE_TwoParams(FOnConfigSectionChanged, const TCHAR* /*IniFilename*/, const TCHAR* /*SectionName*/);
+	static FOnConfigSectionChanged OnConfigSectionChanged;
+
 	DECLARE_MULTICAST_DELEGATE_FourParams(FOnApplyCVarFromIni, const TCHAR* /*SectionName*/, const TCHAR* /*IniFilename*/, uint32 /*SetBy*/, bool /*bAllowCheating*/);
 	static FOnApplyCVarFromIni OnApplyCVarFromIni;
 
