- Change OnConfigSectionChanged to OnConfigSectionsChanged, as it's far cheaper to fire once and have bindees do .Contains on the SectionNames.

diff --git a/Engine/Plugins/Online/EOSShared/Source/EOSShared/Private/EOSSDKManager.cpp b/Engine/Plugins/Online/EOSShared/Source/EOSShared/Private/EOSSDKManager.cpp
--- a/Engine/Plugins/Online/EOSShared/Source/EOSShared/Private/EOSSDKManager.cpp
+++ b/Engine/Plugins/Online/EOSShared/Source/EOSShared/Private/EOSSDKManager.cpp
@@ -194,7 +194,7 @@ EOS_EResult FEOSSDKManager::Initialize()
 		{
 			bInitialized = true;
 
-			FCoreDelegates::OnConfigSectionChanged.AddRaw(this, &FEOSSDKManager::OnConfigSectionChanged);
+			FCoreDelegates::OnConfigSectionsChanged.AddRaw(this, &FEOSSDKManager::OnConfigSectionsChanged);
 			LoadConfig();
 
 #if !NO_LOGGING
@@ -250,9 +250,9 @@ IEOSPlatformHandlePtr FEOSSDKManager::CreatePlatform(EOS_Platform_Options& Platf
 	return SharedPlatform;
 }
 
-void FEOSSDKManager::OnConfigSectionChanged(const TCHAR* IniFilename, const TCHAR* SectionName)
+void FEOSSDKManager::OnConfigSectionsChanged(const FString& IniFilename, const TSet<FString>& SectionNames)
 {
-	if (!FCString::Strcmp(IniFilename, *GEngineIni) && !FCString::Strcmp(SectionName, TEXT("EOSSDK")))
+	if (IniFilename == GEngineIni && SectionNames.Contains(TEXT("EOSSDK")))
 	{
 		LoadConfig();
 	}
@@ -380,7 +380,7 @@ void FEOSSDKManager::Shutdown()
 			ReleaseReleasedPlatforms();
 		}
 
-		FCoreDelegates::OnConfigSectionChanged.RemoveAll(this);
+		FCoreDelegates::OnConfigSectionsChanged.RemoveAll(this);
 
 #if !NO_LOGGING
 		FCoreDelegates::OnLogVerbosityChanged.RemoveAll(this);
diff --git a/Engine/Plugins/Online/EOSShared/Source/EOSShared/Private/EOSSDKManager.h b/Engine/Plugins/Online/EOSShared/Source/EOSShared/Private/EOSSDKManager.h
--- a/Engine/Plugins/Online/EOSShared/Source/EOSShared/Private/EOSSDKManager.h
+++ b/Engine/Plugins/Online/EOSShared/Source/EOSShared/Private/EOSSDKManager.h
@@ -40,7 +40,7 @@ protected:
 private:
 	friend struct FEOSPlatformHandle;
 
-	void OnConfigSectionChanged(const TCHAR* IniFilename, const TCHAR* SectionName);
+	void OnConfigSectionsChanged(const FString& IniFilename, const TSet<FString>& SectionName);
 	void LoadConfig();
 	void ReleasePlatform(EOS_HPlatform PlatformHandle);
 	void ReleaseReleasedPlatforms();
diff --git a/Engine/Plugins/Online/OnlineFramework/Source/Hotfix/Private/OnlineHotfixManager.cpp b/Engine/Plugins/Online/OnlineFramework/Source/Hotfix/Private/OnlineHotfixManager.cpp
--- a/Engine/Plugins/Online/OnlineFramework/Source/Hotfix/Private/OnlineHotfixManager.cpp
+++ b/Engine/Plugins/Online/OnlineFramework/Source/Hotfix/Private/OnlineHotfixManager.cpp
@@ -977,7 +977,7 @@ bool UOnlineHotfixManager::HotfixIniFile(const FString& FileName, const FString&
 
 				if (bIsEngineIni)
 				{
-					// TODO replace all of this with bindees to FCoreDelegates::OnConfigSectionChanged
+					// TODO replace all of this with bindees to FCoreDelegates::OnConfigSectionsChanged
 					const TCHAR* LogConfigSection = TEXT("[Core.Log]");
 					const TCHAR* ConsoleVariableSection = TEXT("[ConsoleVariables]");
 					const TCHAR* HttpSection = TEXT("[HTTP"); // note "]" omitted on purpose since we want a partial match
@@ -1102,10 +1102,7 @@ bool UOnlineHotfixManager::HotfixIniFile(const FString& FileName, const FString&
 	}
 
 	const FString ConfigFileName = ConfigFile->Name.ToString();
-	for (const FString& SectionName : UpdatedSectionNames)
-	{
-		FCoreDelegates::OnConfigSectionChanged.Broadcast(*ConfigFileName, *SectionName);
-	}
+	FCoreDelegates::OnConfigSectionsChanged.Broadcast(ConfigFileName, UpdatedSectionNames);
 
 	// Reload log suppression if configs changed
 	if (bUpdateLogSuppression)
diff --git a/Engine/Source/Runtime/Core/Private/Misc/CoreDelegates.cpp b/Engine/Source/Runtime/Core/Private/Misc/CoreDelegates.cpp
--- a/Engine/Source/Runtime/Core/Private/Misc/CoreDelegates.cpp
+++ b/Engine/Source/Runtime/Core/Private/Misc/CoreDelegates.cpp
@@ -143,7 +143,7 @@ FCoreDelegates::FOnFConfigFileCreated FCoreDelegates::OnFConfigDeleted;
 FCoreDelegates::FOnConfigValueRead FCoreDelegates::OnConfigValueRead;
 FCoreDelegates::FOnConfigSectionRead FCoreDelegates::OnConfigSectionRead;
 FCoreDelegates::FOnConfigSectionRead FCoreDelegates::OnConfigSectionNameRead;
-FCoreDelegates::FOnConfigSectionChanged FCoreDelegates::OnConfigSectionChanged;
+FCoreDelegates::FOnConfigSectionsChanged FCoreDelegates::OnConfigSectionsChanged;
 FCoreDelegates::FOnApplyCVarFromIni FCoreDelegates::OnApplyCVarFromIni;
 FCoreDelegates::FOnSystemResolutionChanged FCoreDelegates::OnSystemResolutionChanged;
 
diff --git a/Engine/Source/Runtime/Core/Public/Misc/CoreDelegates.h b/Engine/Source/Runtime/Core/Public/Misc/CoreDelegates.h
--- a/Engine/Source/Runtime/Core/Public/Misc/CoreDelegates.h
+++ b/Engine/Source/Runtime/Core/Public/Misc/CoreDelegates.h
@@ -343,8 +343,8 @@ public:
 	static FOnConfigSectionRead OnConfigSectionRead;
 	static FOnConfigSectionRead OnConfigSectionNameRead;
 
-	DECLARE_MULTICAST_DELEGATE_TwoParams(FOnConfigSectionChanged, const TCHAR* /*IniFilename*/, const TCHAR* /*SectionName*/);
-	static FOnConfigSectionChanged OnConfigSectionChanged;
+	DECLARE_MULTICAST_DELEGATE_TwoParams(FOnConfigSectionsChanged, const FString& /*IniFilename*/, const TSet<FString>& /*SectionNames*/);
+	static FOnConfigSectionsChanged OnConfigSectionsChanged;
 
 	DECLARE_MULTICAST_DELEGATE_FourParams(FOnApplyCVarFromIni, const TCHAR* /*SectionName*/, const TCHAR* /*IniFilename*/, uint32 /*SetBy*/, bool /*bAllowCheating*/);
 	static FOnApplyCVarFromIni OnApplyCVarFromIni;
