- Copy EOSSharedTypes.h from //UE5/Main CL 19203799
- Add bIsGameThreadCallback condition, as some callbacks occur on audio threads.

diff --git a/Engine/Plugins/Online/EOSShared/Source/EOSShared/Public/EOSSharedTypes.h b/Engine/Plugins/Online/EOSShared/Source/EOSShared/Public/EOSSharedTypes.h
--- a/Engine/Plugins/Online/EOSShared/Source/EOSShared/Public/EOSSharedTypes.h
+++ b/Engine/Plugins/Online/EOSShared/Source/EOSShared/Public/EOSSharedTypes.h
@@ -47,17 +47,23 @@ public:
 		return &CallbackImpl;
 	}
 
+	/** Is this callback intended for the game thread */
+	bool bIsGameThreadCallback = true;
+
 private:
 	/** The object that needs to be checked for lifetime before calling the callback */
 	TWeakPtr<OwningType> Owner;
 
 	static void EOS_CALL CallbackImpl(const CallbackType* Data)
 	{
-		check(IsInGameThread());
-
 		TEOSGlobalCallback* CallbackThis = (TEOSGlobalCallback*)Data->ClientData;
 		check(CallbackThis);
 
+		if (CallbackThis->bIsGameThreadCallback)
+		{
+			check(IsInGameThread());
+		}
+
 		if (CallbackThis->Owner.IsValid())
 		{
 			check(CallbackThis->CallbackLambda);
