- Fix spurious warnings about EOSSDK failing to load.

diff --git a/Engine/Plugins/Online/EOSShared/Source/EOSShared/Private/EOSSDKManager.cpp b/Engine/Plugins/Online/EOSShared/Source/EOSShared/Private/EOSSDKManager.cpp
--- a/Engine/Plugins/Online/EOSShared/Source/EOSShared/Private/EOSSDKManager.cpp
+++ b/Engine/Plugins/Online/EOSShared/Source/EOSShared/Private/EOSSDKManager.cpp
@@ -110,13 +110,19 @@ namespace
 	{
 		void* Result = nullptr;
 
-		const FString ProjectBinaryPath = FPaths::Combine(FPaths::ProjectDir(), TEXT("Binaries"), FPlatformProcess::GetBinariesSubdirectory(), TEXT(EOSSDK_RUNTIME_LIBRARY_NAME));
-		Result = FPlatformProcess::GetDllHandle(*ProjectBinaryPath);
+		const FString ProjectBinaryPath = FPaths::ConvertRelativePathToFull(FPaths::Combine(FPaths::ProjectDir(), TEXT("Binaries"), FPlatformProcess::GetBinariesSubdirectory(), TEXT(EOSSDK_RUNTIME_LIBRARY_NAME)));
+		if (FPaths::FileExists(ProjectBinaryPath))
+		{
+			Result = FPlatformProcess::GetDllHandle(*ProjectBinaryPath);
+		}
 
 		if (!Result)
 		{
-			const FString EngineBinaryPath = FPaths::Combine(FPaths::EngineDir(), TEXT("Binaries"), FPlatformProcess::GetBinariesSubdirectory(), TEXT(EOSSDK_RUNTIME_LIBRARY_NAME));
-			Result = FPlatformProcess::GetDllHandle(*EngineBinaryPath);
+			const FString EngineBinaryPath = FPaths::ConvertRelativePathToFull(FPaths::Combine(FPaths::EngineDir(), TEXT("Binaries"), FPlatformProcess::GetBinariesSubdirectory(), TEXT(EOSSDK_RUNTIME_LIBRARY_NAME)));
+			if (FPaths::FileExists(EngineBinaryPath))
+			{
+				Result = FPlatformProcess::GetDllHandle(*EngineBinaryPath);
+			}
 		}
 
 		if (!Result)
