- Add LexFromString and EOSProductUserIdFromString taking TCHAR* to EOSShared

diff --git a/Engine/Plugins/Online/EOSShared/Source/EOSShared/Private/EOSShared.cpp b/Engine/Plugins/Online/EOSShared/Source/EOSShared/Private/EOSShared.cpp
--- a/Engine/Plugins/Online/EOSShared/Source/EOSShared/Private/EOSShared.cpp
+++ b/Engine/Plugins/Online/EOSShared/Source/EOSShared/Private/EOSShared.cpp
@@ -29,6 +29,11 @@ FString LexToString(const EOS_ProductUserId UserId)
 	return Result;
 }
 
+void LexFromString(EOS_ProductUserId& UserId, const TCHAR* String)
+{
+	UserId = EOS_ProductUserId_FromString(TCHAR_TO_UTF8(String));
+}
+
 FString LexToString(const EOS_EpicAccountId AccountId)
 {
 	FString Result;
diff --git a/Engine/Plugins/Online/EOSShared/Source/EOSShared/Public/EOSShared.h b/Engine/Plugins/Online/EOSShared/Source/EOSShared/Public/EOSShared.h
--- a/Engine/Plugins/Online/EOSShared/Source/EOSShared/Public/EOSShared.h
+++ b/Engine/Plugins/Online/EOSShared/Source/EOSShared/Public/EOSShared.h
@@ -30,7 +30,16 @@ DECLARE_LOG_CATEGORY_EXTERN(LogEOSSDK, Log, All);
 
 EOSSHARED_API FString LexToString(const EOS_EResult EosResult);
 EOSSHARED_API FString LexToString(const EOS_ProductUserId UserId);
+EOSSHARED_API void LexFromString(EOS_ProductUserId& UserId, const TCHAR* String);
+inline EOS_ProductUserId EOSProductUserIdFromString(const TCHAR* String)
+{
+	EOS_ProductUserId UserId;
+	LexFromString(UserId, String);
+	return UserId;
+}
+
 EOSSHARED_API FString LexToString(const EOS_EpicAccountId AccountId);
+
 EOSSHARED_API const TCHAR* LexToString(const EOS_EFriendsStatus FriendStatus);
 EOSSHARED_API bool LexFromString(EOS_EAuthScopeFlags& OutEnum, const TCHAR* InString);
 EOSSHARED_API bool LexFromString(EOS_EExternalCredentialType& OutEnum, const TCHAR* InString);
