Add console command support to EOSShared plugin and add 'info' command for logging configured and cached SDK state.

diff --git a/Engine/Plugins/Online/EOSShared/Source/EOSShared/Private/EOSSDKManager.cpp b/Engine/Plugins/Online/EOSShared/Source/EOSShared/Private/EOSSDKManager.cpp
--- a/Engine/Plugins/Online/EOSShared/Source/EOSShared/Private/EOSSDKManager.cpp
+++ b/Engine/Plugins/Online/EOSShared/Source/EOSShared/Private/EOSSDKManager.cpp
@@ -20,9 +20,14 @@
 
 #include "EOSShared.h"
 
+#include "eos_auth.h"
+#include "eos_connect.h"
+#include "eos_friends.h"
 #include "eos_init.h"
 #include "eos_logging.h"
+#include "eos_presence.h"
 #include "eos_sdk.h"
+#include "eos_userinfo.h"
 #include "eos_version.h"
 
 #ifndef EOS_TRACE_MALLOC
@@ -375,6 +380,34 @@ FString FEOSSDKManager::GetCacheDirBase() const
 	return FPlatformProcess::UserDir();
 }
 
+FString FEOSSDKManager::GetOverrideCountryCode(const EOS_HPlatform Platform) const
+{
+	FString Result;
+
+	char CountryCode[EOS_COUNTRYCODE_MAX_LENGTH + 1];
+	int32_t CountryCodeLength = sizeof(CountryCode);
+	if (EOS_Platform_GetOverrideCountryCode(Platform, CountryCode, &CountryCodeLength) == EOS_EResult::EOS_Success)
+	{
+		Result = UTF8_TO_TCHAR(CountryCode);
+	}
+
+	return Result;
+}
+
+FString FEOSSDKManager::GetOverrideLocaleCode(const EOS_HPlatform Platform) const
+{
+	FString Result;
+
+	char LocaleCode[EOS_LOCALECODE_MAX_LENGTH + 1];
+	int32_t LocaleCodeLength = sizeof(LocaleCode);
+	if (EOS_Platform_GetOverrideLocaleCode(Platform, LocaleCode, &LocaleCodeLength) == EOS_EResult::EOS_Success)
+	{
+		Result = UTF8_TO_TCHAR(LocaleCode);
+	}
+
+	return Result;
+}
+
 void FEOSSDKManager::ReleasePlatform(EOS_HPlatform PlatformHandle)
 {
 	if (ensure(ActivePlatforms.Contains(PlatformHandle)
@@ -435,6 +468,345 @@ EOS_EResult FEOSSDKManager::EOSInitialize(EOS_InitializeOptions& Options)
 	return EOS_Initialize(&Options);
 }
 
+bool FEOSSDKManager::Exec(class UWorld* InWorld, const TCHAR* Cmd, FOutputDevice& Ar)
+{
+	if (!FParse::Command(&Cmd, TEXT("EOSSDK")))
+	{
+		return false;
+	}
+
+	if (FParse::Command(&Cmd, TEXT("INFO")))
+	{
+		LogInfo();
+	}
+	else
+	{
+		UE_LOG(LogEOSSDK, Warning, TEXT("Unknown exec command: %s]"), Cmd);
+	}
+
+	return true;
+}
+
+#define UE_LOG_EOSSDK_INFO(Format, ...) UE_LOG(LogEOSSDK, Log, TEXT("Info: %*s" Format), Indent * 2, TEXT(""), ##__VA_ARGS__)
+
+void FEOSSDKManager::LogInfo(int32 Indent) const
+{
+	UE_LOG_EOSSDK_INFO("ProductName=%s", *GetProductName());
+	UE_LOG_EOSSDK_INFO("ProductVersion=%s", *GetProductVersion());
+	UE_LOG_EOSSDK_INFO("CacheDirBase=%s", *GetCacheDirBase());
+	UE_LOG_EOSSDK_INFO("Platforms=%d", ActivePlatforms.Num());
+
+	for (int32 PlatformIndex = 0; PlatformIndex < ActivePlatforms.Num(); PlatformIndex++)
+	{
+		const EOS_HPlatform Platform = ActivePlatforms[PlatformIndex];
+		UE_LOG_EOSSDK_INFO("Platform=%d", PlatformIndex);
+		Indent++;
+		LogPlatformInfo(Platform, Indent);
+		Indent--;
+	}
+}
+
+void FEOSSDKManager::LogPlatformInfo(const EOS_HPlatform Platform, int32 Indent) const
+{
+	UE_LOG_EOSSDK_INFO("ApplicationStatus=%s", LexToString(EOS_Platform_GetApplicationStatus(Platform)));
+	UE_LOG_EOSSDK_INFO("NetworkStatus=%s", LexToString(EOS_Platform_GetNetworkStatus(Platform)));
+	UE_LOG_EOSSDK_INFO("OverrideCountryCode=%s", *GetOverrideCountryCode(Platform));
+	UE_LOG_EOSSDK_INFO("OverrideLocaleCode=%s", *GetOverrideLocaleCode(Platform));
+
+	static_assert(EOS_PLATFORM_GETDESKTOPCROSSPLAYSTATUS_API_LATEST == 1, "EOS_Platform_GetDesktopCrossplayStatusOptions updated");
+	EOS_Platform_GetDesktopCrossplayStatusOptions GetDesktopCrossplayStatusOptions;
+	GetDesktopCrossplayStatusOptions.ApiVersion = EOS_PLATFORM_GETDESKTOPCROSSPLAYSTATUS_API_LATEST;
+
+	EOS_Platform_GetDesktopCrossplayStatusInfo GetDesktopCrossplayStatusInfo;
+	EOS_EResult Result = EOS_Platform_GetDesktopCrossplayStatus(Platform, &GetDesktopCrossplayStatusOptions, &GetDesktopCrossplayStatusInfo);
+	if (Result == EOS_EResult::EOS_Success)
+	{
+		UE_LOG_EOSSDK_INFO("DesktopCrossplayStatusInfo Status=%s ServiceInitResult=%d",
+			LexToString(GetDesktopCrossplayStatusInfo.Status),
+			GetDesktopCrossplayStatusInfo.ServiceInitResult);
+	}
+	else
+	{
+		UE_LOG_EOSSDK_INFO("DesktopCrossplayStatusInfo (EOS_Platform_GetDesktopCrossplayStatus failed: %s)", *LexToString(Result));
+	}
+
+	const EOS_HAuth AuthHandle = EOS_Platform_GetAuthInterface(Platform);
+	const int32_t AuthLoggedInAccountsCount = EOS_Auth_GetLoggedInAccountsCount(AuthHandle);
+	UE_LOG_EOSSDK_INFO("AuthLoggedInAccounts=%d", AuthLoggedInAccountsCount);
+
+	for (int32_t AuthLoggedInAccountIndex = 0; AuthLoggedInAccountIndex < AuthLoggedInAccountsCount; AuthLoggedInAccountIndex++)
+	{
+		const EOS_EpicAccountId LoggedInAccount = EOS_Auth_GetLoggedInAccountByIndex(AuthHandle, AuthLoggedInAccountIndex);
+		UE_LOG_EOSSDK_INFO("AuthLoggedInAccount=%d", AuthLoggedInAccountIndex);
+		Indent++;
+		LogUserInfo(Platform, LoggedInAccount, LoggedInAccount, Indent);
+		LogAuthInfo(Platform, LoggedInAccount, Indent);
+		LogPresenceInfo(Platform, LoggedInAccount, LoggedInAccount, Indent);
+		LogFriendsInfo(Platform, LoggedInAccount, Indent);
+		Indent--;
+	}
+
+	const EOS_HConnect ConnectHandle = EOS_Platform_GetConnectInterface(Platform);
+	const int32_t ConnectLoggedInAccountsCount = EOS_Connect_GetLoggedInUsersCount(ConnectHandle);
+	UE_LOG_EOSSDK_INFO("ConnectLoggedInAccounts=%d", ConnectLoggedInAccountsCount);
+
+	for (int32_t ConnectLoggedInAccountIndex = 0; ConnectLoggedInAccountIndex < ConnectLoggedInAccountsCount; ConnectLoggedInAccountIndex++)
+	{
+		const EOS_ProductUserId LoggedInAccount = EOS_Connect_GetLoggedInUserByIndex(ConnectHandle, ConnectLoggedInAccountIndex);
+		UE_LOG_EOSSDK_INFO("ConnectLoggedInAccount=%d", ConnectLoggedInAccountIndex);
+		Indent++;
+		LogConnectInfo(Platform, LoggedInAccount, Indent);
+		Indent--;
+	}
+}
+
+void FEOSSDKManager::LogAuthInfo(const EOS_HPlatform Platform, const EOS_EpicAccountId LoggedInAccount, int32 Indent) const
+{
+	const EOS_HAuth AuthHandle = EOS_Platform_GetAuthInterface(Platform);
+	UE_LOG_EOSSDK_INFO("LoginStatus=%s", LexToString(EOS_Auth_GetLoginStatus(AuthHandle, LoggedInAccount)));
+
+	static_assert(EOS_AUTH_COPYUSERAUTHTOKEN_API_LATEST == 1, "EOS_Auth_CopyUserAuthTokenOptions updated");
+	EOS_Auth_CopyUserAuthTokenOptions CopyUserAuthTokenOptions;
+	CopyUserAuthTokenOptions.ApiVersion = EOS_AUTH_COPYUSERAUTHTOKEN_API_LATEST;
+
+	EOS_Auth_Token* AuthToken;
+	EOS_EResult Result = EOS_Auth_CopyUserAuthToken(AuthHandle, &CopyUserAuthTokenOptions, LoggedInAccount, &AuthToken);
+	if (Result == EOS_EResult::EOS_Success)
+	{
+		UE_LOG_EOSSDK_INFO("AuthToken");
+		Indent++;
+		UE_LOG_EOSSDK_INFO("App=%s", UTF8_TO_TCHAR(AuthToken->App));
+		UE_LOG_EOSSDK_INFO("ClientId=%s", UTF8_TO_TCHAR(AuthToken->ClientId));
+#if !UE_BUILD_SHIPPING
+		UE_LOG_EOSSDK_INFO("AccessToken=%s", UTF8_TO_TCHAR(AuthToken->AccessToken));
+#endif // !UE_BUILD_SHIPPING
+		UE_LOG_EOSSDK_INFO("ExpiresIn=%f", AuthToken->ExpiresIn);
+		UE_LOG_EOSSDK_INFO("ExpiresAt=%s", UTF8_TO_TCHAR(AuthToken->ExpiresAt));
+		UE_LOG_EOSSDK_INFO("AuthType=%s", LexToString(AuthToken->AuthType));
+#if !UE_BUILD_SHIPPING
+		UE_LOG_EOSSDK_INFO("RefreshToken=%s", UTF8_TO_TCHAR(AuthToken->RefreshToken));
+#endif // !UE_BUILD_SHIPPING
+		UE_LOG_EOSSDK_INFO("RefreshExpiresIn=%f", AuthToken->RefreshExpiresIn);
+		UE_LOG_EOSSDK_INFO("RefreshExpiresAt=%s", UTF8_TO_TCHAR(AuthToken->RefreshExpiresAt));
+		Indent--;
+		EOS_Auth_Token_Release(AuthToken);
+	}
+	else
+	{
+		UE_LOG_EOSSDK_INFO("AuthToken (EOS_Auth_CopyUserAuthToken failed: %s)", *LexToString(Result));
+	}
+
+#if !UE_BUILD_SHIPPING
+	static_assert(EOS_AUTH_COPYIDTOKEN_API_LATEST == 1, "EOS_Auth_CopyIdTokenOptions updated");
+	EOS_Auth_CopyIdTokenOptions CopyIdTokenOptions;
+	CopyIdTokenOptions.ApiVersion = EOS_AUTH_COPYIDTOKEN_API_LATEST;
+	CopyIdTokenOptions.AccountId = LoggedInAccount;
+
+	EOS_Auth_IdToken* IdToken;
+	Result = EOS_Auth_CopyIdToken(AuthHandle, &CopyIdTokenOptions, &IdToken);
+	if (Result == EOS_EResult::EOS_Success)
+	{
+		UE_LOG_EOSSDK_INFO("IdToken=%s", UTF8_TO_TCHAR(IdToken->JsonWebToken));
+		EOS_Auth_IdToken_Release(IdToken);
+	}
+	else
+	{
+		UE_LOG_EOSSDK_INFO("IdToken (EOS_Auth_CopyIdToken failed: %s)", *LexToString(Result));
+	}
+#endif // !UE_BUILD_SHIPPING
+}
+
+void FEOSSDKManager::LogUserInfo(const EOS_HPlatform Platform, const EOS_EpicAccountId LoggedInAccount, const EOS_EpicAccountId TargetAccount, int32 Indent) const
+{
+	UE_LOG_EOSSDK_INFO("EpicAccountId=%s", *LexToString(TargetAccount));
+
+	static_assert(EOS_USERINFO_COPYUSERINFO_API_LATEST >=2 && EOS_USERINFO_COPYUSERINFO_API_LATEST <= 3, "EOS_UserInfo_CopyUserInfo updated");
+	EOS_UserInfo_CopyUserInfoOptions CopyUserInfoOptions;
+	CopyUserInfoOptions.ApiVersion = EOS_USERINFO_COPYUSERINFO_API_LATEST;
+	CopyUserInfoOptions.LocalUserId = LoggedInAccount;
+	CopyUserInfoOptions.TargetUserId = TargetAccount;
+
+	const EOS_HUserInfo UserInfoHandle = EOS_Platform_GetUserInfoInterface(Platform);
+	EOS_UserInfo* UserInfo;
+	EOS_EResult Result = EOS_UserInfo_CopyUserInfo(UserInfoHandle, &CopyUserInfoOptions, &UserInfo);
+	if (Result == EOS_EResult::EOS_Success)
+	{
+		UE_LOG_EOSSDK_INFO("UserInfo");
+		Indent++;
+		UE_LOG_EOSSDK_INFO("Country=%s", UTF8_TO_TCHAR(UserInfo->Country));
+		UE_LOG_EOSSDK_INFO("DisplayName=%s", UTF8_TO_TCHAR(UserInfo->DisplayName));
+		UE_LOG_EOSSDK_INFO("PreferredLanguage=%s", UTF8_TO_TCHAR(UserInfo->PreferredLanguage));
+		UE_LOG_EOSSDK_INFO("Nickname=%s", UTF8_TO_TCHAR(UserInfo->Nickname));
+		UE_LOG_EOSSDK_INFO("DisplayNameSanitized=%s", UTF8_TO_TCHAR(UserInfo->DisplayNameSanitized));
+		Indent--;
+
+		EOS_UserInfo_Release(UserInfo);
+	}
+	else
+	{
+		UE_LOG_EOSSDK_INFO("UserInfo (EOS_UserInfo_CopyUserInfo failed: %s)", *LexToString(Result));
+	}
+}
+
+void FEOSSDKManager::LogPresenceInfo(const EOS_HPlatform Platform, const EOS_EpicAccountId LoggedInAccount, const EOS_EpicAccountId TargetAccount, int32 Indent) const
+{
+	static_assert(EOS_PRESENCE_HASPRESENCE_API_LATEST == 1, "EOS_Presence_HasPresenceOptions updated");
+	EOS_Presence_HasPresenceOptions HasPresenceOptions;
+	HasPresenceOptions.ApiVersion = EOS_PRESENCE_HASPRESENCE_API_LATEST;
+	HasPresenceOptions.LocalUserId = LoggedInAccount;
+	HasPresenceOptions.TargetUserId = TargetAccount;
+
+	const EOS_HPresence PresenceHandle = EOS_Platform_GetPresenceInterface(Platform);
+	if (!EOS_Presence_HasPresence(PresenceHandle, &HasPresenceOptions))
+	{
+		UE_LOG_EOSSDK_INFO("Presence (None)");
+		return;
+	}
+
+	static_assert(EOS_PRESENCE_COPYPRESENCE_API_LATEST >=2 && EOS_PRESENCE_COPYPRESENCE_API_LATEST <= 3, "EOS_Presence_CopyPresenceOptions updated");
+	EOS_Presence_CopyPresenceOptions CopyPresenceOptions;
+	CopyPresenceOptions.ApiVersion = EOS_PRESENCE_COPYPRESENCE_API_LATEST;
+	CopyPresenceOptions.LocalUserId = LoggedInAccount;
+	CopyPresenceOptions.TargetUserId = TargetAccount;
+
+	EOS_Presence_Info* PresenceInfo;
+	EOS_EResult Result = EOS_Presence_CopyPresence(PresenceHandle, &CopyPresenceOptions, &PresenceInfo);
+	if (Result == EOS_EResult::EOS_Success)
+	{
+		UE_LOG_EOSSDK_INFO("Presence");
+		Indent++;
+		UE_LOG_EOSSDK_INFO("Status=%s", LexToString(PresenceInfo->Status));
+		UE_LOG_EOSSDK_INFO("ProductId=%s", UTF8_TO_TCHAR(PresenceInfo->ProductId));
+		UE_LOG_EOSSDK_INFO("ProductName=%s", UTF8_TO_TCHAR(PresenceInfo->ProductName));
+		UE_LOG_EOSSDK_INFO("ProductVersion=%s", UTF8_TO_TCHAR(PresenceInfo->ProductVersion));
+		UE_LOG_EOSSDK_INFO("Platform=%s", UTF8_TO_TCHAR(PresenceInfo->Platform));
+		UE_LOG_EOSSDK_INFO("IntegratedPlatform=%s", UTF8_TO_TCHAR(PresenceInfo->IntegratedPlatform));
+		UE_LOG_EOSSDK_INFO("RichText=%s", UTF8_TO_TCHAR(PresenceInfo->RichText));
+		UE_LOG_EOSSDK_INFO("RecordsCount=%d", PresenceInfo->RecordsCount);
+		Indent++;
+		for (int32_t Index = 0; Index < PresenceInfo->RecordsCount; Index++)
+		{
+			UE_LOG_EOSSDK_INFO("Key=%s Value=%s", UTF8_TO_TCHAR(PresenceInfo->Records[Index].Key), UTF8_TO_TCHAR(PresenceInfo->Records[Index].Value));
+		}
+		Indent--;
+		Indent--;
+
+		EOS_Presence_Info_Release(PresenceInfo);
+	}
+	else
+	{
+		UE_LOG_EOSSDK_INFO("Presence (EOS_Presence_CopyPresence failed: %s)", *LexToString(Result));
+	}
+}
+
+void FEOSSDKManager::LogFriendsInfo(const EOS_HPlatform Platform, const EOS_EpicAccountId LoggedInAccount, int32 Indent) const
+{
+	static_assert(EOS_FRIENDS_GETFRIENDSCOUNT_API_LATEST == 1, "EOS_Friends_GetFriendsCountOptions updated");
+	EOS_Friends_GetFriendsCountOptions FriendsCountOptions;
+	FriendsCountOptions.ApiVersion = EOS_FRIENDS_GETFRIENDSCOUNT_API_LATEST;
+	FriendsCountOptions.LocalUserId = LoggedInAccount;
+
+	const EOS_HFriends FriendsHandle = EOS_Platform_GetFriendsInterface(Platform);
+	const int32_t FriendsCount = EOS_Friends_GetFriendsCount(FriendsHandle, &FriendsCountOptions);
+	UE_LOG_EOSSDK_INFO("Friends=%d", FriendsCount);
+
+	static_assert(EOS_FRIENDS_GETFRIENDATINDEX_API_LATEST == 1, "EOS_Friends_GetFriendAtIndexOptions updated");
+	EOS_Friends_GetFriendAtIndexOptions FriendAtIndexOptions;
+	FriendAtIndexOptions.ApiVersion = EOS_FRIENDS_GETFRIENDATINDEX_API_LATEST;
+	FriendAtIndexOptions.LocalUserId = LoggedInAccount;
+
+	for (int32_t FriendIndex = 0; FriendIndex < FriendsCount; FriendIndex++)
+	{
+		FriendAtIndexOptions.Index = FriendIndex;
+		const EOS_EpicAccountId FriendId = EOS_Friends_GetFriendAtIndex(FriendsHandle, &FriendAtIndexOptions);
+		UE_LOG_EOSSDK_INFO("Friend=%d", FriendIndex);
+		Indent++;
+
+		static_assert(EOS_FRIENDS_GETSTATUS_API_LATEST == 1, "EOS_Friends_GetStatusOptions updated");
+		EOS_Friends_GetStatusOptions GetStatusOptions;
+		GetStatusOptions.ApiVersion = EOS_FRIENDS_GETSTATUS_API_LATEST;
+		GetStatusOptions.LocalUserId = LoggedInAccount;
+		GetStatusOptions.TargetUserId = FriendId;
+
+		const EOS_EFriendsStatus FriendStatus = EOS_Friends_GetStatus(FriendsHandle, &GetStatusOptions);
+		UE_LOG_EOSSDK_INFO("FriendStatus=%s", LexToString(FriendStatus));
+
+		LogUserInfo(Platform, LoggedInAccount, FriendId, Indent);
+		LogPresenceInfo(Platform, LoggedInAccount, FriendId, Indent);
+		Indent--;
+	}
+}
+
+void FEOSSDKManager::LogConnectInfo(const EOS_HPlatform Platform, const EOS_ProductUserId LoggedInAccount, int32 Indent) const
+{
+	UE_LOG_EOSSDK_INFO("ProductUserId=%s", *LexToString(LoggedInAccount));
+
+	const EOS_HConnect ConnectHandle = EOS_Platform_GetConnectInterface(Platform);
+	UE_LOG_EOSSDK_INFO("LoginStatus=%s", LexToString(EOS_Connect_GetLoginStatus(ConnectHandle, LoggedInAccount)));
+
+	EOS_EResult Result;
+
+#if !UE_BUILD_SHIPPING
+	static_assert(EOS_CONNECT_COPYIDTOKEN_API_LATEST == 1, "EOS_Connect_CopyIdTokenOptions updated");
+	EOS_Connect_CopyIdTokenOptions CopyIdTokenOptions;
+	CopyIdTokenOptions.ApiVersion = EOS_CONNECT_COPYIDTOKEN_API_LATEST;
+	CopyIdTokenOptions.LocalUserId = LoggedInAccount;
+
+	EOS_Connect_IdToken* IdToken;
+	Result = EOS_Connect_CopyIdToken(ConnectHandle, &CopyIdTokenOptions, &IdToken);
+	if (Result == EOS_EResult::EOS_Success)
+	{
+		UE_LOG_EOSSDK_INFO("IdToken=%s", UTF8_TO_TCHAR(IdToken->JsonWebToken));
+		EOS_Connect_IdToken_Release(IdToken);
+	}
+	else
+	{
+		UE_LOG_EOSSDK_INFO("IdToken (EOS_Connect_CopyIdToken failed: %s)", *LexToString(Result));
+	}
+#endif // !UE_BUILD_SHIPPING
+
+	static_assert(EOS_CONNECT_GETPRODUCTUSEREXTERNALACCOUNTCOUNT_API_LATEST == 1, "EOS_Connect_GetProductUserExternalAccountCountOptions updated");
+	EOS_Connect_GetProductUserExternalAccountCountOptions ExternalAccountCountOptions;
+	ExternalAccountCountOptions.ApiVersion = EOS_CONNECT_GETPRODUCTUSEREXTERNALACCOUNTCOUNT_API_LATEST;
+	ExternalAccountCountOptions.TargetUserId = LoggedInAccount;
+
+	const uint32_t ExternalAccountCount = EOS_Connect_GetProductUserExternalAccountCount(ConnectHandle, &ExternalAccountCountOptions);
+	UE_LOG_EOSSDK_INFO("ExternalAccounts=%d", ExternalAccountCount);
+
+	static_assert(EOS_CONNECT_COPYPRODUCTUSEREXTERNALACCOUNTBYINDEX_API_LATEST == 1, "EOS_Connect_CopyProductUserExternalAccountByIndexOptions updated");
+	EOS_Connect_CopyProductUserExternalAccountByIndexOptions ExternalAccountByIndexOptions;
+	ExternalAccountByIndexOptions.ApiVersion = EOS_CONNECT_COPYPRODUCTUSEREXTERNALACCOUNTBYINDEX_API_LATEST;
+	ExternalAccountByIndexOptions.TargetUserId = LoggedInAccount;
+
+	for (uint32_t ExternalAccountIndex = 0; ExternalAccountIndex < ExternalAccountCount; ExternalAccountIndex++)
+	{
+		UE_LOG_EOSSDK_INFO("ExternalAccount=%u", ExternalAccountIndex);
+		Indent++;
+
+		ExternalAccountByIndexOptions.ExternalAccountInfoIndex = ExternalAccountIndex;
+		EOS_Connect_ExternalAccountInfo* ExternalAccountInfo;
+		Result = EOS_Connect_CopyProductUserExternalAccountByIndex(ConnectHandle, &ExternalAccountByIndexOptions, &ExternalAccountInfo);
+		if (Result == EOS_EResult::EOS_Success)
+		{
+			UE_LOG_EOSSDK_INFO("ExternalAccountInfo");
+			Indent++;
+			UE_LOG_EOSSDK_INFO("DisplayName=%s", UTF8_TO_TCHAR(ExternalAccountInfo->DisplayName));
+			UE_LOG_EOSSDK_INFO("AccountId=%s", UTF8_TO_TCHAR(ExternalAccountInfo->AccountId));
+			UE_LOG_EOSSDK_INFO("AccountIdType=%s", LexToString(ExternalAccountInfo->AccountIdType));
+			UE_LOG_EOSSDK_INFO("LastLoginTime=%lld", ExternalAccountInfo->LastLoginTime);
+			Indent--;
+
+			EOS_Connect_ExternalAccountInfo_Release(ExternalAccountInfo);
+		}
+		else
+		{
+			UE_LOG_EOSSDK_INFO("ExternalAccountInfo (EOS_Connect_CopyProductUserExternalAccountByIndex failed: %s)", *LexToString(Result));
+		}
+
+		Indent--;
+	}
+}
+
 FEOSPlatformHandle::~FEOSPlatformHandle()
 {
 	Manager.ReleasePlatform(PlatformHandle);
@@ -448,4 +820,44 @@ void FEOSPlatformHandle::Tick()
 	EOS_Platform_Tick(PlatformHandle);
 }
 
+FString FEOSPlatformHandle::GetOverrideCountryCode() const
+{
+	return Manager.GetOverrideCountryCode(PlatformHandle);
+}
+
+FString FEOSPlatformHandle::GetOverrideLocaleCode() const
+{
+	return Manager.GetOverrideLocaleCode(PlatformHandle);
+}
+
+void FEOSPlatformHandle::LogInfo(int32 Indent) const
+{
+	Manager.LogPlatformInfo(PlatformHandle, Indent);
+}
+
+void FEOSPlatformHandle::LogAuthInfo(const EOS_EpicAccountId LoggedInAccount, int32 Indent) const
+{
+	Manager.LogAuthInfo(PlatformHandle, LoggedInAccount, Indent);
+}
+
+void FEOSPlatformHandle::LogUserInfo(const EOS_EpicAccountId LoggedInAccount, const EOS_EpicAccountId TargetAccount, int32 Indent) const
+{
+	Manager.LogUserInfo(PlatformHandle, LoggedInAccount, TargetAccount, Indent);
+}
+
+void FEOSPlatformHandle::LogPresenceInfo(const EOS_EpicAccountId LoggedInAccount, const EOS_EpicAccountId TargetAccount, int32 Indent) const
+{
+	Manager.LogPresenceInfo(PlatformHandle, LoggedInAccount, TargetAccount, Indent);
+}
+
+void FEOSPlatformHandle::LogFriendsInfo(const EOS_EpicAccountId LoggedInAccount, int32 Indent) const
+{
+	Manager.LogFriendsInfo(PlatformHandle, LoggedInAccount, Indent);
+}
+
+void FEOSPlatformHandle::LogConnectInfo(const EOS_ProductUserId LoggedInAccount, int32 Indent) const
+{
+	Manager.LogConnectInfo(PlatformHandle, LoggedInAccount, Indent);
+}
+
 #endif // WITH_EOS_SDK
\ No newline at end of file
diff --git a/Engine/Plugins/Online/EOSShared/Source/EOSShared/Private/EOSSDKManager.h b/Engine/Plugins/Online/EOSShared/Source/EOSShared/Private/EOSSDKManager.h
--- a/Engine/Plugins/Online/EOSShared/Source/EOSShared/Private/EOSSDKManager.h
+++ b/Engine/Plugins/Online/EOSShared/Source/EOSShared/Private/EOSSDKManager.h
@@ -11,12 +11,17 @@
 #if defined(EOS_PLATFORM_BASE_FILE_NAME)
 #include EOS_PLATFORM_BASE_FILE_NAME
 #endif
+
+#include "eos_auth_types.h"
 #include "eos_common.h"
+#include "eos_connect_types.h"
 #include "eos_init.h"
 
 struct FEOSPlatformHandle;
 
-class FEOSSDKManager : public IEOSSDKManager
+class FEOSSDKManager :
+	public IEOSSDKManager,
+	public FSelfRegisteringExec
 {
 public:
 	FEOSSDKManager();
@@ -31,8 +36,22 @@ public:
 	virtual FString GetProductName() const override;
 	virtual FString GetProductVersion() const override;
 	virtual FString GetCacheDirBase() const override;
+	virtual FString GetOverrideCountryCode(const EOS_HPlatform Platform) const override;
+	virtual FString GetOverrideLocaleCode(const EOS_HPlatform Platform) const override;
+
+	virtual void LogInfo(int32 Indent = 0) const override;
+	virtual void LogPlatformInfo(const EOS_HPlatform Platform, int32 Indent = 0) const override;
+	virtual void LogAuthInfo(const EOS_HPlatform Platform, const EOS_EpicAccountId LoggedInAccount, int32 Indent = 0) const override;
+	virtual void LogUserInfo(const EOS_HPlatform Platform, const EOS_EpicAccountId LoggedInAccount, const EOS_EpicAccountId TargetAccount, int32 Indent = 0) const override;
+	virtual void LogPresenceInfo(const EOS_HPlatform Platform, const EOS_EpicAccountId LoggedInAccount, const EOS_EpicAccountId TargetAccount, int32 Indent = 0) const override;
+	virtual void LogFriendsInfo(const EOS_HPlatform Platform, const EOS_EpicAccountId LoggedInAccount, int32 Indent = 0) const override;
+	virtual void LogConnectInfo(const EOS_HPlatform Platform, const EOS_ProductUserId LoggedInAccount, int32 Indent = 0) const override;
 	// End IEOSSDKManager
 
+	// Begin FSelfRegisteringExec
+	virtual bool Exec(class UWorld* InWorld, const TCHAR* Cmd, FOutputDevice& Ar) override;
+	// End FSelfRegisteringExec
+
 	void Shutdown();
 
 protected:
@@ -79,6 +98,16 @@ struct FEOSPlatformHandle : public IEOSPlatformHandle
 	virtual ~FEOSPlatformHandle();
 	virtual void Tick() override;
 
+	virtual FString GetOverrideCountryCode() const override;
+	virtual FString GetOverrideLocaleCode() const override;
+
+	virtual void LogInfo(int32 Indent = 0) const override;
+	virtual void LogAuthInfo(const EOS_EpicAccountId LoggedInAccount, int32 Indent = 0) const override;
+	virtual void LogUserInfo(const EOS_EpicAccountId LoggedInAccount, const EOS_EpicAccountId TargetAccount, int32 Indent = 0) const override;
+	virtual void LogPresenceInfo(const EOS_EpicAccountId LoggedInAccount, const EOS_EpicAccountId TargetAccount, int32 Indent = 0) const override;
+	virtual void LogFriendsInfo(const EOS_EpicAccountId LoggedInAccount, int32 Indent = 0) const override;
+	virtual void LogConnectInfo(const EOS_ProductUserId LoggedInAccount, int32 Indent = 0) const override;
+
 	/* Reference to the EOSSDK manager */
 	FEOSSDKManager& Manager;
 };
diff --git a/Engine/Plugins/Online/EOSShared/Source/EOSShared/Private/EOSShared.cpp b/Engine/Plugins/Online/EOSShared/Source/EOSShared/Private/EOSShared.cpp
--- a/Engine/Plugins/Online/EOSShared/Source/EOSShared/Private/EOSShared.cpp
+++ b/Engine/Plugins/Online/EOSShared/Source/EOSShared/Private/EOSShared.cpp
@@ -5,6 +5,8 @@
 
 #include "eos_auth_types.h"
 #include "eos_friends_types.h"
+#include "eos_presence_types.h"
+#include "eos_types.h"
 
 DEFINE_LOG_CATEGORY(LogEOSSDK);
 
@@ -50,6 +52,67 @@ FString LexToString(const EOS_EpicAccountId AccountId)
 	return Result;
 }
 
+const TCHAR* LexToString(const EOS_EApplicationStatus ApplicationStatus)
+{
+	switch (ApplicationStatus)
+	{
+		case EOS_EApplicationStatus::EOS_AS_BackgroundConstrained:		return TEXT("BackgroundConstrained");
+		case EOS_EApplicationStatus::EOS_AS_BackgroundUnconstrained:	return TEXT("BackgroundUnconstrained");
+		case EOS_EApplicationStatus::EOS_AS_BackgroundSuspended:		return TEXT("BackgroundSuspended");
+		case EOS_EApplicationStatus::EOS_AS_Foreground:					return TEXT("Foreground");
+		default: checkNoEntry();										return TEXT("Unknown");
+	}
+}
+
+const TCHAR* LexToString(const EOS_EAuthTokenType AuthTokenType)
+{
+	switch (AuthTokenType)
+	{
+		case EOS_EAuthTokenType::EOS_ATT_Client:	return TEXT("Client");
+		case EOS_EAuthTokenType::EOS_ATT_User:		return TEXT("User");
+		default: checkNoEntry();					return TEXT("Unknown");
+	}
+}
+
+const TCHAR* LexToString(const EOS_EDesktopCrossplayStatus DesktopCrossplayStatus)
+{
+	switch (DesktopCrossplayStatus)
+	{
+		case EOS_EDesktopCrossplayStatus::EOS_DCS_OK:							return TEXT("OK");
+		case EOS_EDesktopCrossplayStatus::EOS_DCS_ApplicationNotBootstrapped:	return TEXT("ApplicationNotBootstrapped");
+		case EOS_EDesktopCrossplayStatus::EOS_DCS_ServiceNotInstalled:			return TEXT("ServiceNotInstalled");
+		case EOS_EDesktopCrossplayStatus::EOS_DCS_ServiceStartFailed:			return TEXT("ServiceStartFailed");
+		case EOS_EDesktopCrossplayStatus::EOS_DCS_ServiceNotRunning:			return TEXT("ServiceNotRunning");
+		case EOS_EDesktopCrossplayStatus::EOS_DCS_OverlayDisabled:				return TEXT("OverlayDisabled");
+		case EOS_EDesktopCrossplayStatus::EOS_DCS_OverlayNotInstalled:			return TEXT("OverlayNotInstalled");
+		case EOS_EDesktopCrossplayStatus::EOS_DCS_OverlayTrustCheckFailed:		return TEXT("OverlayTrustCheckFailed");
+		case EOS_EDesktopCrossplayStatus::EOS_DCS_OverlayLoadFailed:			return TEXT("OverlayLoadFailed");
+		default: checkNoEntry();												return TEXT("Unknown");
+	}
+}
+
+const TCHAR* LexToString(const EOS_EExternalAccountType ExternalAccountType)
+{
+	switch (ExternalAccountType)
+	{
+		case EOS_EExternalAccountType::EOS_EAT_EPIC:		return TEXT("Epic");
+		case EOS_EExternalAccountType::EOS_EAT_STEAM:		return TEXT("Steam");
+		case EOS_EExternalAccountType::EOS_EAT_PSN:			return TEXT("PSN");
+		case EOS_EExternalAccountType::EOS_EAT_XBL:			return TEXT("XBL");
+		case EOS_EExternalAccountType::EOS_EAT_DISCORD:		return TEXT("Discord");
+		case EOS_EExternalAccountType::EOS_EAT_GOG:			return TEXT("GOG");
+		case EOS_EExternalAccountType::EOS_EAT_NINTENDO:	return TEXT("Nintendo");
+		case EOS_EExternalAccountType::EOS_EAT_UPLAY:		return TEXT("UPlay");
+		case EOS_EExternalAccountType::EOS_EAT_OPENID:		return TEXT("OpenID");
+		case EOS_EExternalAccountType::EOS_EAT_APPLE:		return TEXT("Apple");
+		case EOS_EExternalAccountType::EOS_EAT_GOOGLE:		return TEXT("Google");
+		case EOS_EExternalAccountType::EOS_EAT_OCULUS:		return TEXT("Oculus");
+		case EOS_EExternalAccountType::EOS_EAT_ITCHIO:		return TEXT("ItchIO");
+		case EOS_EExternalAccountType::EOS_EAT_AMAZON:		return TEXT("Amazon");
+		default: checkNoEntry();							return TEXT("Unknown");
+	}
+}
+
 const TCHAR* LexToString(const EOS_EFriendsStatus FriendStatus)
 {
 	switch (FriendStatus)
@@ -62,14 +125,38 @@ const TCHAR* LexToString(const EOS_EFriendsStatus FriendStatus)
 	}
 }
 
-const TCHAR* LexToString(EOS_ELoginStatus LoginStatus)
+const TCHAR* LexToString(const EOS_ELoginStatus LoginStatus)
 {
 	switch (LoginStatus)
 	{
-	default: checkNoEntry(); // Intentional fallthrough
-	case EOS_ELoginStatus::EOS_LS_NotLoggedIn:			return TEXT("EOS_LS_NotLoggedIn");
-	case EOS_ELoginStatus::EOS_LS_UsingLocalProfile:	return TEXT("EOS_LS_UsingLocalProfile");
-	case EOS_ELoginStatus::EOS_LS_LoggedIn:				return TEXT("EOS_LS_LoggedIn");
+		case EOS_ELoginStatus::EOS_LS_NotLoggedIn:			return TEXT("NotLoggedIn");
+		case EOS_ELoginStatus::EOS_LS_UsingLocalProfile:	return TEXT("UsingLocalProfile");
+		case EOS_ELoginStatus::EOS_LS_LoggedIn:				return TEXT("LoggedIn");
+		default: checkNoEntry();							return TEXT("Unknown");
+	}
+}
+
+const TCHAR* LexToString(const EOS_ENetworkStatus NetworkStatus)
+{
+	switch (NetworkStatus)
+	{
+		case EOS_ENetworkStatus::EOS_NS_Disabled:	return TEXT("Disabled");
+		case EOS_ENetworkStatus::EOS_NS_Offline:	return TEXT("Offline");
+		case EOS_ENetworkStatus::EOS_NS_Online:		return TEXT("Online");
+		default: checkNoEntry();					return TEXT("Unknown");
+	}
+}
+
+const TCHAR* LexToString(const EOS_Presence_EStatus PresenceStatus)
+{
+	switch (PresenceStatus)
+	{
+		case EOS_Presence_EStatus::EOS_PS_Offline:		return TEXT("Offline");
+		case EOS_Presence_EStatus::EOS_PS_Online:		return TEXT("Online");
+		case EOS_Presence_EStatus::EOS_PS_Away:			return TEXT("Away");
+		case EOS_Presence_EStatus::EOS_PS_ExtendedAway:	return TEXT("ExtendedAway");
+		case EOS_Presence_EStatus::EOS_PS_DoNotDisturb:	return TEXT("DoNotDisturb");
+		default: checkNoEntry();						return TEXT("Unknown");
 	}
 }
 
diff --git a/Engine/Plugins/Online/EOSShared/Source/EOSShared/Public/EOSShared.h b/Engine/Plugins/Online/EOSShared/Source/EOSShared/Public/EOSShared.h
--- a/Engine/Plugins/Online/EOSShared/Source/EOSShared/Public/EOSShared.h
+++ b/Engine/Plugins/Online/EOSShared/Source/EOSShared/Public/EOSShared.h
@@ -21,9 +21,14 @@
 #endif
 
 #define EOS_ENUM_FORWARD_DECL(name) enum class name : int32_t;
+EOS_ENUM_FORWARD_DECL(EOS_EApplicationStatus);
 EOS_ENUM_FORWARD_DECL(EOS_EAuthScopeFlags);
+EOS_ENUM_FORWARD_DECL(EOS_EAuthTokenType);
+EOS_ENUM_FORWARD_DECL(EOS_EDesktopCrossplayStatus);
 EOS_ENUM_FORWARD_DECL(EOS_EFriendsStatus);
 EOS_ENUM_FORWARD_DECL(EOS_ELoginCredentialType);
+EOS_ENUM_FORWARD_DECL(EOS_ENetworkStatus);
+EOS_ENUM_FORWARD_DECL(EOS_Presence_EStatus);
 #undef EOS_ENUM_FORWARD_DECL
 
 DECLARE_LOG_CATEGORY_EXTERN(LogEOSSDK, Log, All);
@@ -40,8 +45,15 @@ inline EOS_ProductUserId EOSProductUserIdFromString(const TCHAR* String)
 
 EOSSHARED_API FString LexToString(const EOS_EpicAccountId AccountId);
 
+EOSSHARED_API const TCHAR* LexToString(const EOS_EApplicationStatus ApplicationStatus);
+EOSSHARED_API const TCHAR* LexToString(const EOS_EAuthTokenType AuthTokenType);
+EOSSHARED_API const TCHAR* LexToString(const EOS_EDesktopCrossplayStatus DesktopCrossplayStatus);
+EOSSHARED_API const TCHAR* LexToString(const EOS_EExternalAccountType ExternalAccountType);
 EOSSHARED_API const TCHAR* LexToString(const EOS_EFriendsStatus FriendStatus);
 EOSSHARED_API const TCHAR* LexToString(const EOS_ELoginStatus LoginStatus);
+EOSSHARED_API const TCHAR* LexToString(const EOS_ENetworkStatus NetworkStatus);
+EOSSHARED_API const TCHAR* LexToString(const EOS_Presence_EStatus PresenceStatus);
+
 EOSSHARED_API bool LexFromString(EOS_EAuthScopeFlags& OutEnum, const TCHAR* InString);
 EOSSHARED_API bool LexFromString(EOS_EExternalCredentialType& OutEnum, const TCHAR* InString);
 EOSSHARED_API bool LexFromString(EOS_ELoginCredentialType& OutEnum, const TCHAR* InString);
diff --git a/Engine/Plugins/Online/EOSShared/Source/EOSShared/Public/IEOSSDKManager.h b/Engine/Plugins/Online/EOSShared/Source/EOSShared/Public/IEOSSDKManager.h
--- a/Engine/Plugins/Online/EOSShared/Source/EOSShared/Public/IEOSSDKManager.h
+++ b/Engine/Plugins/Online/EOSShared/Source/EOSShared/Public/IEOSSDKManager.h
@@ -25,6 +25,16 @@ public:
 
 	operator EOS_HPlatform() const { return PlatformHandle; }
 
+	virtual FString GetOverrideCountryCode() const = 0;
+	virtual FString GetOverrideLocaleCode() const = 0;
+
+	virtual void LogInfo(int32 Indent = 0) const = 0;
+	virtual void LogAuthInfo(const EOS_EpicAccountId LoggedInAccount, int32 Indent = 0) const = 0;
+	virtual void LogUserInfo(const EOS_EpicAccountId LoggedInAccount, const EOS_EpicAccountId TargetAccount, int32 Indent = 0) const = 0;
+	virtual void LogPresenceInfo(const EOS_EpicAccountId LoggedInAccount, const EOS_EpicAccountId TargetAccount, int32 Indent = 0) const = 0;
+	virtual void LogFriendsInfo(const EOS_EpicAccountId LoggedInAccount, int32 Indent = 0) const = 0;
+	virtual void LogConnectInfo(const EOS_ProductUserId LoggedInAccount, int32 Indent = 0) const = 0;
+
 protected:
 	EOS_HPlatform PlatformHandle;
 };
@@ -58,6 +68,16 @@ public:
 	virtual FString GetProductName() const = 0;
 	virtual FString GetProductVersion() const = 0;
 	virtual FString GetCacheDirBase() const = 0;
+	virtual FString GetOverrideCountryCode(const EOS_HPlatform Platform) const = 0;
+	virtual FString GetOverrideLocaleCode(const EOS_HPlatform Platform) const = 0;
+
+	virtual void LogInfo(int32 Indent = 0) const = 0;
+	virtual void LogPlatformInfo(const EOS_HPlatform Platform, int32 Indent = 0) const = 0;
+	virtual void LogAuthInfo(const EOS_HPlatform Platform, const EOS_EpicAccountId LoggedInAccount, int32 Indent = 0) const = 0;
+	virtual void LogUserInfo(const EOS_HPlatform Platform, const EOS_EpicAccountId LoggedInAccount, const EOS_EpicAccountId TargetAccount, int32 Indent = 0) const = 0;
+	virtual void LogPresenceInfo(const EOS_HPlatform Platform, const EOS_EpicAccountId LoggedInAccount, const EOS_EpicAccountId TargetAccount, int32 Indent = 0) const = 0;
+	virtual void LogFriendsInfo(const EOS_HPlatform Platform, const EOS_EpicAccountId LoggedInAccount, int32 Indent = 0) const = 0;
+	virtual void LogConnectInfo(const EOS_HPlatform Platform, const EOS_ProductUserId LoggedInAccount, int32 Indent = 0) const = 0;
 
 	FEOSSDKManagerOnPreInitializeSDK OnPreInitializeSDK;
 	FEOSSDKManagerOnPreCreatePlatform OnPreCreatePlatform;
