- Don't create/register EOSSDKManager in commandlet runs.

diff --git a/Engine/Plugins/Online/EOSShared/Source/EOSShared/Private/EOSSharedModule.cpp b/Engine/Plugins/Online/EOSShared/Source/EOSShared/Private/EOSSharedModule.cpp
--- a/Engine/Plugins/Online/EOSShared/Source/EOSShared/Private/EOSSharedModule.cpp
+++ b/Engine/Plugins/Online/EOSShared/Source/EOSShared/Private/EOSSharedModule.cpp
@@ -6,6 +6,8 @@
 #include "Misc/ConfigCacheIni.h"
 #include "Misc/CoreDelegates.h"
 #include "Modules/ModuleManager.h"
+#include "CoreGlobals.h"
+#include "EOSShared.h"
 
 #include COMPILED_PLATFORM_HEADER(EOSSDKManager.h)
 
@@ -16,6 +18,13 @@ IMPLEMENT_MODULE(FEOSSharedModule, EOSShared);
 void FEOSSharedModule::StartupModule()
 {
 #if WITH_EOS_SDK
+	if (IsRunningCommandlet())
+	{
+		// No need to do anything when running in commandlet mode. We won't register the EOSSDKManager, EOS will not be initialized, no platforms will be created, etc.
+		UE_LOG(LogEOSSDK, Log, TEXT("IsRunningCommandlet=true, skipping EOSSDK initialization."))
+		return;
+	}
+
 	SDKManager = MakeUnique<FPlatformEOSSDKManager>();
 	check(SDKManager);
 
