Add named platform configuration to EOSShared along with default platform support.

diff --git a/Engine/Plugins/Online/EOSShared/Source/EOSShared/Private/EOSSDKManager.cpp b/Engine/Plugins/Online/EOSShared/Source/EOSShared/Private/EOSSDKManager.cpp
--- a/Engine/Plugins/Online/EOSShared/Source/EOSShared/Private/EOSSDKManager.cpp
+++ b/Engine/Plugins/Online/EOSShared/Source/EOSShared/Private/EOSSDKManager.cpp
@@ -254,6 +254,180 @@ EOS_EResult FEOSSDKManager::Initialize()
 	}
 }
 
+const FEOSSDKPlatformConfig* FEOSSDKManager::GetPlatformConfig(const FString& PlatformConfigName, bool bLoadIfMissing)
+{
+	if (PlatformConfigName.IsEmpty())
+	{
+		UE_LOG(LogEOSSDK, Warning, TEXT("Platform name can't be empty"));
+		return nullptr;
+	}
+
+	FEOSSDKPlatformConfig* PlatformConfig = PlatformConfigs.Find(PlatformConfigName);
+	if (PlatformConfig || !bLoadIfMissing)
+	{
+		return PlatformConfig;
+	}
+
+	const FString SectionName(TEXT("EOSSDK.Platform.") + PlatformConfigName);
+	if (!GConfig->DoesSectionExist(*SectionName, GEngineIni))
+	{
+		UE_LOG(LogEOSSDK, Warning, TEXT("Could not find platform config: %s"), *PlatformConfigName);
+		return nullptr;
+	}
+
+	PlatformConfig = &PlatformConfigs.Emplace(PlatformConfigName);
+
+	PlatformConfig->Name = PlatformConfigName;
+	GConfig->GetString(*SectionName, TEXT("ProductId"), PlatformConfig->ProductId, GEngineIni);
+	GConfig->GetString(*SectionName, TEXT("SandboxId"), PlatformConfig->SandboxId, GEngineIni);
+	GConfig->GetString(*SectionName, TEXT("ClientId"), PlatformConfig->ClientId, GEngineIni);
+	GConfig->GetString(*SectionName, TEXT("ClientSecret"), PlatformConfig->ClientSecret, GEngineIni);
+	GConfig->GetString(*SectionName, TEXT("EncryptionKey"), PlatformConfig->EncryptionKey, GEngineIni);
+	GConfig->GetString(*SectionName, TEXT("OverrideCountryCode"), PlatformConfig->OverrideCountryCode, GEngineIni);
+	GConfig->GetString(*SectionName, TEXT("OverrideLocaleCode"), PlatformConfig->OverrideLocaleCode, GEngineIni);
+	GConfig->GetString(*SectionName, TEXT("DeploymentId"), PlatformConfig->DeploymentId, GEngineIni);
+	GConfig->GetString(*SectionName, TEXT("CacheDirectory"), PlatformConfig->CacheDirectory, GEngineIni);
+	GConfig->GetBool(*SectionName, TEXT("bIsServer"), PlatformConfig->bIsServer, GEngineIni);
+	GConfig->GetBool(*SectionName, TEXT("bLoadingInEditor"), PlatformConfig->bLoadingInEditor, GEngineIni);
+	GConfig->GetBool(*SectionName, TEXT("bDisableOverlay"), PlatformConfig->bDisableOverlay, GEngineIni);
+	GConfig->GetBool(*SectionName, TEXT("bDisableSocialOverlay"), PlatformConfig->bDisableSocialOverlay, GEngineIni);
+	GConfig->GetBool(*SectionName, TEXT("bWindowsEnableOverlayD3D9"), PlatformConfig->bWindowsEnableOverlayD3D9, GEngineIni);
+	GConfig->GetBool(*SectionName, TEXT("bWindowsEnableOverlayD3D10"), PlatformConfig->bWindowsEnableOverlayD3D10, GEngineIni);
+	GConfig->GetBool(*SectionName, TEXT("bWindowsEnableOverlayOpenGL"), PlatformConfig->bWindowsEnableOverlayOpenGL, GEngineIni);
+	GConfig->GetBool(*SectionName, TEXT("bEnableRTC"), PlatformConfig->bEnableRTC, GEngineIni);
+	GConfig->GetInt(*SectionName, TEXT("TickBudgetInMilliseconds"), PlatformConfig->TickBudgetInMilliseconds, GEngineIni);
+	GConfig->GetArray(*SectionName, TEXT("OptionalConfig"), PlatformConfig->OptionalConfig, GEngineIni);
+
+	UE_LOG(LogEOSSDK, Verbose, TEXT("Loaded platform config: %s"), *PlatformConfigName);
+	return PlatformConfig;
+}
+
+bool FEOSSDKManager::AddPlatformConfig(const FEOSSDKPlatformConfig& PlatformConfig)
+{
+	if (PlatformConfig.Name.IsEmpty())
+	{
+		UE_LOG(LogEOSSDK, Warning, TEXT("Platform name can't be empty"));
+		return false;
+	}
+
+	if (PlatformConfigs.Find(PlatformConfig.Name))
+	{
+		UE_LOG(LogEOSSDK, Warning, TEXT("Platform config already exists: %s"), *PlatformConfig.Name);
+		return false;
+	}
+
+	PlatformConfigs.Emplace(PlatformConfig.Name, PlatformConfig);
+	UE_LOG(LogEOSSDK, Verbose, TEXT("Added platform config: %s"), *PlatformConfig.Name);
+	return true;
+}
+
+const FString& FEOSSDKManager::GetDefaultPlatformConfigName()
+{
+	if (DefaultPlatformConfigName.IsEmpty())
+	{
+		FString PlatformConfigName;
+		if (GConfig->GetString(TEXT("EOSSDK"), TEXT("DefaultPlatformConfigName"), PlatformConfigName, GEngineIni))
+		{
+			SetDefaultPlatformConfigName(PlatformConfigName);
+		}
+	}
+
+	return DefaultPlatformConfigName;
+}
+
+void FEOSSDKManager::SetDefaultPlatformConfigName(const FString& PlatformConfigName)
+{
+	if (DefaultPlatformConfigName != PlatformConfigName)
+	{
+		UE_LOG(LogEOSSDK, Verbose, TEXT("Default platform name changed: New=%s Old=%s"), *PlatformConfigName, *DefaultPlatformConfigName);
+		OnDefaultPlatformConfigNameChanged.Broadcast(PlatformConfigName, DefaultPlatformConfigName);
+		DefaultPlatformConfigName = PlatformConfigName;
+	}
+}
+
+IEOSPlatformHandlePtr FEOSSDKManager::CreatePlatform(const FString& PlatformConfigName)
+{
+	if (PlatformConfigName.IsEmpty())
+	{
+		UE_LOG(LogEOSSDK, Warning, TEXT("Platform name can't be empty"));
+		return IEOSPlatformHandlePtr();
+	}
+
+	IEOSPlatformHandleWeakPtr* WeakPlatformHandle = PlatformHandles.Find(PlatformConfigName);
+	if (WeakPlatformHandle)
+	{
+		if (WeakPlatformHandle->IsValid())
+		{
+			UE_LOG(LogEOSSDK, Verbose, TEXT("Found existing platform handle: %s"), *PlatformConfigName);
+			return WeakPlatformHandle->Pin();
+		}
+
+		UE_LOG(LogEOSSDK, Verbose, TEXT("Removing stale platform handle pointer: %s"), *PlatformConfigName);
+		PlatformHandles.Remove(PlatformConfigName);
+	}
+
+	const FEOSSDKPlatformConfig* const PlatformConfig = GetPlatformConfig(PlatformConfigName, true);
+	if (!PlatformConfig)
+	{
+		return IEOSPlatformHandlePtr();
+	}
+
+	const FTCHARToUTF8 Utf8ProductId(*PlatformConfig->ProductId);
+	const FTCHARToUTF8 Utf8SandboxId(*PlatformConfig->SandboxId);
+	const FTCHARToUTF8 Utf8ClientId(*PlatformConfig->ClientId);
+	const FTCHARToUTF8 Utf8ClientSecret(*PlatformConfig->ClientSecret);
+	const FTCHARToUTF8 Utf8EncryptionKey(*PlatformConfig->EncryptionKey);
+	const FTCHARToUTF8 Utf8OverrideCountryCode(*PlatformConfig->OverrideCountryCode);
+	const FTCHARToUTF8 Utf8OverrideLocaleCode(*PlatformConfig->OverrideLocaleCode);
+	const FTCHARToUTF8 Utf8DeploymentId(*PlatformConfig->DeploymentId);
+	const FTCHARToUTF8 Utf8CacheDirectory(*PlatformConfig->CacheDirectory);
+
+	static_assert(EOS_PLATFORM_OPTIONS_API_LATEST == 12, "EOS_Platform_Options updated");
+	EOS_Platform_Options PlatformOptions = {};
+	PlatformOptions.ApiVersion = EOS_PLATFORM_OPTIONS_API_LATEST;
+	PlatformOptions.Reserved = nullptr;
+	PlatformOptions.ProductId = Utf8ProductId.Length() ? Utf8ProductId.Get() : nullptr;
+	PlatformOptions.SandboxId = Utf8SandboxId.Length() ? Utf8SandboxId.Get() : nullptr;
+	PlatformOptions.ClientCredentials.ClientId = Utf8ClientId.Length() ? Utf8ClientId.Get() : nullptr;
+	PlatformOptions.ClientCredentials.ClientSecret = Utf8ClientSecret.Length() ? Utf8ClientSecret.Get() : nullptr;
+	PlatformOptions.bIsServer = PlatformConfig->bIsServer;
+	PlatformOptions.EncryptionKey = Utf8EncryptionKey.Length() ? Utf8EncryptionKey.Get() : nullptr;
+	PlatformOptions.OverrideCountryCode = Utf8OverrideCountryCode.Length() ? Utf8OverrideCountryCode.Get() : nullptr;
+	PlatformOptions.OverrideLocaleCode = Utf8OverrideLocaleCode.Length() ? Utf8OverrideLocaleCode.Get() : nullptr;
+	PlatformOptions.DeploymentId = Utf8DeploymentId.Length() ? Utf8DeploymentId.Get() : nullptr;
+
+	PlatformOptions.Flags = 0;
+	if (PlatformConfig->bLoadingInEditor) PlatformOptions.Flags |= EOS_PF_LOADING_IN_EDITOR;
+	if (PlatformConfig->bDisableOverlay) PlatformOptions.Flags |= EOS_PF_DISABLE_OVERLAY;
+	if (PlatformConfig->bDisableSocialOverlay) PlatformOptions.Flags |= EOS_PF_DISABLE_SOCIAL_OVERLAY;
+
+	PlatformOptions.CacheDirectory = Utf8CacheDirectory.Length() ? Utf8DeploymentId.Get() : nullptr;
+	PlatformOptions.TickBudgetInMilliseconds = PlatformConfig->TickBudgetInMilliseconds;
+
+	static_assert(EOS_PLATFORM_RTCOPTIONS_API_LATEST == 1, "EOS_Platform_RTCOptions updated");
+	EOS_Platform_RTCOptions PlatformRTCOptions = {};
+	PlatformRTCOptions.ApiVersion = EOS_PLATFORM_RTCOPTIONS_API_LATEST;
+	PlatformRTCOptions.PlatformSpecificOptions = nullptr;
+	PlatformOptions.RTCOptions = PlatformConfig->bEnableRTC ? &PlatformRTCOptions : nullptr;
+
+	PlatformOptions.IntegratedPlatformOptionsContainerHandle = nullptr;
+
+	IEOSPlatformHandlePtr PlatformHandle = CreatePlatform(*PlatformConfig, PlatformOptions);
+	if (PlatformHandle.IsValid())
+	{
+		UE_LOG(LogEOSSDK, Verbose, TEXT("Created platform handle: %s"), *PlatformConfigName);
+		PlatformHandles.Emplace(PlatformConfigName, PlatformHandle);
+	}
+
+	return PlatformHandle;
+}
+
+IEOSPlatformHandlePtr FEOSSDKManager::CreatePlatform(const FEOSSDKPlatformConfig& PlatformConfig, EOS_Platform_Options& PlatformOptions)
+{
+	OnPreCreateNamedPlatform.Broadcast(PlatformConfig, PlatformOptions);
+	return CreatePlatform(PlatformOptions);
+}
+
 IEOSPlatformHandlePtr FEOSSDKManager::CreatePlatform(EOS_Platform_Options& PlatformOptions)
 {
 	IEOSPlatformHandlePtr SharedPlatform;
diff --git a/Engine/Plugins/Online/EOSShared/Source/EOSShared/Private/EOSSDKManager.h b/Engine/Plugins/Online/EOSShared/Source/EOSShared/Private/EOSSDKManager.h
--- a/Engine/Plugins/Online/EOSShared/Source/EOSShared/Private/EOSSDKManager.h
+++ b/Engine/Plugins/Online/EOSShared/Source/EOSShared/Private/EOSSDKManager.h
@@ -32,6 +32,12 @@ public:
 	virtual EOS_EResult Initialize() override;
 	virtual bool IsInitialized() const override { return bInitialized; }
 
+	virtual const FEOSSDKPlatformConfig* GetPlatformConfig(const FString& PlatformConfigName, bool bLoadIfMissing = false) override;
+	virtual bool AddPlatformConfig(const FEOSSDKPlatformConfig& PlatformConfig) override;
+	virtual const FString& GetDefaultPlatformConfigName() override;
+	virtual void SetDefaultPlatformConfigName(const FString& PlatformConfigName) override;
+
+	virtual IEOSPlatformHandlePtr CreatePlatform(const FString& PlatformConfigName) override;
 	virtual IEOSPlatformHandlePtr CreatePlatform(EOS_Platform_Options& PlatformOptions) override;
 
 	virtual FString GetProductName() const override;
@@ -59,6 +65,7 @@ public:
 
 protected:
 	virtual EOS_EResult EOSInitialize(EOS_InitializeOptions& Options);
+	virtual IEOSPlatformHandlePtr CreatePlatform(const FEOSSDKPlatformConfig& PlatformConfig, EOS_Platform_Options& PlatformOptions);
 
 private:
 	friend struct FEOSPlatformHandle;
@@ -87,6 +94,12 @@ private:
 	FTSTicker::FDelegateHandle TickerHandle;
 	/** Callback objects, to be released after EOS_Shutdown */
 	TArray<TUniquePtr<FCallbackBase>> CallbackObjects;
+	/** Cache of named platform configs that have been loaded from ini files or added manually. */
+	TMap<FString, FEOSSDKPlatformConfig> PlatformConfigs;
+	/** Default platform config name to use. */
+	FString DefaultPlatformConfigName;
+	/** Cache of named platform handles that have been created. */
+	TMap<FString, IEOSPlatformHandleWeakPtr> PlatformHandles;
 
 	// Config
 	/** Interval between platform ticks. 0 means we tick every frame. */
diff --git a/Engine/Plugins/Online/EOSShared/Source/EOSShared/Private/Windows/WindowsEOSSDKManager.cpp b/Engine/Plugins/Online/EOSShared/Source/EOSShared/Private/Windows/WindowsEOSSDKManager.cpp
new file mode 100644
--- /dev/null
+++ b/Engine/Plugins/Online/EOSShared/Source/EOSShared/Private/Windows/WindowsEOSSDKManager.cpp
@@ -0,0 +1,34 @@
+// Copyright Epic Games, Inc. All Rights Reserved.
+
+#if WITH_EOS_SDK
+
+#include "WindowsEOSSDKManager.h"
+
+#include "HAL/FileManager.h"
+#include "Misc/Paths.h"
+
+#include "Windows/eos_Windows.h"
+
+IEOSPlatformHandlePtr FWindowsEOSSDKManager::CreatePlatform(const FEOSSDKPlatformConfig& PlatformConfig, EOS_Platform_Options& PlatformOptions)
+{
+	if (PlatformConfig.bWindowsEnableOverlayD3D9) PlatformOptions.Flags |= EOS_PF_WINDOWS_ENABLE_OVERLAY_D3D9;
+	if (PlatformConfig.bWindowsEnableOverlayD3D10) PlatformOptions.Flags |= EOS_PF_WINDOWS_ENABLE_OVERLAY_D3D10;
+	if (PlatformConfig.bWindowsEnableOverlayOpenGL) PlatformOptions.Flags |= EOS_PF_WINDOWS_ENABLE_OVERLAY_OPENGL;
+
+	if (PlatformConfig.bEnableRTC)
+	{
+		static const FString XAudioPath = IFileManager::Get().ConvertToAbsolutePathForExternalAppForRead(*FPaths::Combine(FPaths::EngineDir(), TEXT("Binaries/ThirdParty/Windows/XAudio2_9"), PLATFORM_64BITS ? TEXT("x64") : TEXT("x86"), TEXT("xaudio2_9redist.dll")));
+		static const FTCHARToUTF8 Utf8XAudioPath(*XAudioPath);
+
+		static_assert(EOS_WINDOWS_RTCOPTIONS_API_LATEST == 1, "EOS_Windows_RTCOptions updated");
+		static EOS_Windows_RTCOptions WindowsRTCOptions = {};
+		WindowsRTCOptions.ApiVersion = EOS_WINDOWS_RTCOPTIONS_API_LATEST;
+		WindowsRTCOptions.XAudio29DllPath = Utf8XAudioPath.Get();
+
+		const_cast<EOS_Platform_RTCOptions*>(PlatformOptions.RTCOptions)->PlatformSpecificOptions = &WindowsRTCOptions;
+	}
+
+	return FEOSSDKManager::CreatePlatform(PlatformConfig, PlatformOptions);
+}
+
+#endif // WITH_EOS_SDK
\ No newline at end of file
diff --git a/Engine/Plugins/Online/EOSShared/Source/EOSShared/Private/Windows/WindowsEOSSDKManager.h b/Engine/Plugins/Online/EOSShared/Source/EOSShared/Private/Windows/WindowsEOSSDKManager.h
--- a/Engine/Plugins/Online/EOSShared/Source/EOSShared/Private/Windows/WindowsEOSSDKManager.h
+++ b/Engine/Plugins/Online/EOSShared/Source/EOSShared/Private/Windows/WindowsEOSSDKManager.h
@@ -6,6 +6,12 @@
 
 #include "EOSSDKManager.h"
 
-using FPlatformEOSSDKManager = FEOSSDKManager;
+class FWindowsEOSSDKManager : public FEOSSDKManager
+{
+protected:
+	virtual IEOSPlatformHandlePtr CreatePlatform(const FEOSSDKPlatformConfig& PlatformConfig, EOS_Platform_Options& PlatformOptions) override;
+};
+
+using FPlatformEOSSDKManager = FWindowsEOSSDKManager;
 
 #endif // WITH_EOS_SDK
\ No newline at end of file
diff --git a/Engine/Plugins/Online/EOSShared/Source/EOSShared/Public/IEOSSDKManager.h b/Engine/Plugins/Online/EOSShared/Source/EOSShared/Public/IEOSSDKManager.h
--- a/Engine/Plugins/Online/EOSShared/Source/EOSShared/Public/IEOSSDKManager.h
+++ b/Engine/Plugins/Online/EOSShared/Source/EOSShared/Public/IEOSSDKManager.h
@@ -13,7 +13,33 @@
 #include "eos_init.h"
 #include "eos_types.h"
 
+struct FEOSSDKPlatformConfig
+{
+	FString Name;
+	FString ProductId;
+	FString SandboxId;
+	FString ClientId;
+	FString ClientSecret;
+	FString EncryptionKey;
+	FString OverrideCountryCode;
+	FString OverrideLocaleCode;
+	FString DeploymentId;
+	FString CacheDirectory;
+	bool bIsServer = false;
+	bool bLoadingInEditor = false;
+	bool bDisableOverlay = false;
+	bool bDisableSocialOverlay = false;
+	bool bWindowsEnableOverlayD3D9 = false;
+	bool bWindowsEnableOverlayD3D10 = false;
+	bool bWindowsEnableOverlayOpenGL = false;
+	bool bEnableRTC = true;
+	int32 TickBudgetInMilliseconds = 1;
+	TArray<FString> OptionalConfig;
+};
+
 DECLARE_MULTICAST_DELEGATE_OneParam(FEOSSDKManagerOnPreInitializeSDK, EOS_InitializeOptions& Options);
+DECLARE_MULTICAST_DELEGATE_TwoParams(FEOSSDKManagerOnDefaultPlatformConfigNameChanged, const FString& NewName, const FString& OldName);
+DECLARE_MULTICAST_DELEGATE_TwoParams(FEOSSDKManagerOnPreCreateNamedPlatform, const FEOSSDKPlatformConfig& Config, EOS_Platform_Options& Options);
 DECLARE_MULTICAST_DELEGATE_OneParam(FEOSSDKManagerOnPreCreatePlatform, EOS_Platform_Options& Options);
 
 class IEOSPlatformHandle
@@ -39,7 +65,9 @@ public:
 protected:
 	EOS_HPlatform PlatformHandle;
 };
+
 using IEOSPlatformHandlePtr = TSharedPtr<IEOSPlatformHandle, ESPMode::ThreadSafe>;
+using IEOSPlatformHandleWeakPtr = TWeakPtr<IEOSPlatformHandle, ESPMode::ThreadSafe>;
 
 class IEOSSDKManager : public IModularFeature
 {
@@ -64,6 +92,18 @@ public:
 	virtual EOS_EResult Initialize() = 0;
 	virtual bool IsInitialized() const = 0;
 
+	virtual const FEOSSDKPlatformConfig* GetPlatformConfig(const FString& PlatformConfigName, bool bLoadIfMissing = false) = 0;
+	virtual bool AddPlatformConfig(const FEOSSDKPlatformConfig& PlatformConfig) = 0;
+	virtual const FString& GetDefaultPlatformConfigName() = 0;
+	virtual void SetDefaultPlatformConfigName(const FString& PlatformConfigName) = 0;
+
+	/**
+	 * Create a platform handle for a platform config name. Config is loaded from .ini files if it was not added with AddPlatformConfig.
+	 * If a platform handle already exists for the config name, this will return a shared pointer to that handle and not create a new one.
+	 */
+	virtual IEOSPlatformHandlePtr CreatePlatform(const FString& PlatformConfigName) = 0;
+
+	/** Create a platform handle using EOSSDK options directly. */
 	virtual IEOSPlatformHandlePtr CreatePlatform(EOS_Platform_Options& PlatformOptions) = 0;
 
 	virtual FString GetProductName() const = 0;
@@ -84,6 +124,8 @@ public:
 	virtual void AddCallbackObject(TUniquePtr<class FCallbackBase> CallbackObj) = 0;
 
 	FEOSSDKManagerOnPreInitializeSDK OnPreInitializeSDK;
+	FEOSSDKManagerOnDefaultPlatformConfigNameChanged OnDefaultPlatformConfigNameChanged;
+	FEOSSDKManagerOnPreCreateNamedPlatform OnPreCreateNamedPlatform;
 	FEOSSDKManagerOnPreCreatePlatform OnPreCreatePlatform;
 };
 
