- Fix EOSSDK errors being logged when we create but do not tick a platform before destroying it.

diff --git a/Engine/Plugins/Online/EOSShared/Source/EOSShared/Private/EOSSDKManager.cpp b/Engine/Plugins/Online/EOSShared/Source/EOSShared/Private/EOSSDKManager.cpp
--- a/Engine/Plugins/Online/EOSShared/Source/EOSShared/Private/EOSSDKManager.cpp
+++ b/Engine/Plugins/Online/EOSShared/Source/EOSShared/Private/EOSSDKManager.cpp
@@ -446,6 +446,9 @@ IEOSPlatformHandlePtr FEOSSDKManager::CreatePlatform(EOS_Platform_Options& Platf
 			// TODO Hardcode the application and network state for now
 			EOS_Platform_SetApplicationStatus(PlatformHandle, EOS_EApplicationStatus::EOS_AS_Foreground);
 			EOS_Platform_SetNetworkStatus(PlatformHandle, EOS_ENetworkStatus::EOS_NS_Online);
+
+			// Tick the platform once to work around EOSSDK error logging that occurs if you create then immediately destroy a platform.
+			SharedPlatform->Tick();
 		}
 		else
 		{
