- Make private fields/methods protected so they are accessible from subclasses.

diff --git a/Engine/Plugins/Online/EOSShared/Source/EOSShared/Private/EOSSDKManager.h b/Engine/Plugins/Online/EOSShared/Source/EOSShared/Private/EOSSDKManager.h
--- a/Engine/Plugins/Online/EOSShared/Source/EOSShared/Private/EOSSDKManager.h
+++ b/Engine/Plugins/Online/EOSShared/Source/EOSShared/Private/EOSSDKManager.h
@@ -67,7 +67,6 @@ protected:
 	virtual EOS_EResult EOSInitialize(EOS_InitializeOptions& Options);
 	virtual IEOSPlatformHandlePtr CreatePlatform(const FEOSSDKPlatformConfig& PlatformConfig, EOS_Platform_Options& PlatformOptions);
 
-private:
 	friend struct FEOSPlatformHandle;
 
 	void OnConfigSectionsChanged(const FString& IniFilename, const TSet<FString>& SectionName);
