- Add switch support for setting Application/Network state.

diff --git a/Engine/Plugins/Online/EOSShared/Source/EOSShared/Private/EOSSDKManager.cpp b/Engine/Plugins/Online/EOSShared/Source/EOSShared/Private/EOSSDKManager.cpp
--- a/Engine/Plugins/Online/EOSShared/Source/EOSShared/Private/EOSSDKManager.cpp
+++ b/Engine/Plugins/Online/EOSShared/Source/EOSShared/Private/EOSSDKManager.cpp
@@ -443,9 +443,8 @@ IEOSPlatformHandlePtr FEOSSDKManager::CreatePlatform(EOS_Platform_Options& Platf
 			SharedPlatform = MakeShared<FEOSPlatformHandle, ESPMode::ThreadSafe>(*this, PlatformHandle);
 			SetupTicker();
 
-			// TODO Hardcode the application and network state for now
-			EOS_Platform_SetApplicationStatus(PlatformHandle, EOS_EApplicationStatus::EOS_AS_Foreground);
-			EOS_Platform_SetNetworkStatus(PlatformHandle, EOS_ENetworkStatus::EOS_NS_Online);
+			EOS_Platform_SetApplicationStatus(PlatformHandle, CachedApplicationStatus);
+			EOS_Platform_SetNetworkStatus(PlatformHandle, CachedNetworkStatus);
 
 			// Tick the platform once to work around EOSSDK error logging that occurs if you create then immediately destroy a platform.
 			SharedPlatform->Tick();
@@ -523,6 +522,26 @@ bool FEOSSDKManager::Tick(float)
 	return true;
 }
 
+void FEOSSDKManager::OnApplicationStatusChanged(EOS_EApplicationStatus ApplicationStatus)
+{
+	UE_LOG(LogEOSSDK, Log, TEXT("OnApplicationStatusChanged [%s] -> [%s]"), *LexToString(CachedApplicationStatus), *LexToString(ApplicationStatus));
+	CachedApplicationStatus = ApplicationStatus;
+	for (EOS_HPlatform PlatformHandle : ActivePlatforms)
+	{
+		EOS_Platform_SetApplicationStatus(PlatformHandle, ApplicationStatus);
+	}
+}
+
+void FEOSSDKManager::OnNetworkStatusChanged(EOS_ENetworkStatus NetworkStatus)
+{
+	UE_LOG(LogEOSSDK, Log, TEXT("OnNetworkStatusChanged [%s] -> [%s]"), *LexToString(CachedNetworkStatus), *LexToString(NetworkStatus));
+	CachedNetworkStatus = NetworkStatus;
+	for (EOS_HPlatform PlatformHandle : ActivePlatforms)
+	{
+		EOS_Platform_SetNetworkStatus(PlatformHandle, NetworkStatus);
+	}
+}
+
 void FEOSSDKManager::OnLogVerbosityChanged(const FLogCategoryName& CategoryName, ELogVerbosity::Type OldVerbosity, ELogVerbosity::Type NewVerbosity)
 {
 #if !NO_LOGGING
diff --git a/Engine/Plugins/Online/EOSShared/Source/EOSShared/Private/EOSSDKManager.h b/Engine/Plugins/Online/EOSShared/Source/EOSShared/Private/EOSSDKManager.h
--- a/Engine/Plugins/Online/EOSShared/Source/EOSShared/Private/EOSSDKManager.h
+++ b/Engine/Plugins/Online/EOSShared/Source/EOSShared/Private/EOSSDKManager.h
@@ -66,6 +66,12 @@ public:
 protected:
 	virtual EOS_EResult EOSInitialize(EOS_InitializeOptions& Options);
 	virtual IEOSPlatformHandlePtr CreatePlatform(const FEOSSDKPlatformConfig& PlatformConfig, EOS_Platform_Options& PlatformOptions);
+	virtual bool Tick(float);
+
+	void OnApplicationStatusChanged(EOS_EApplicationStatus ApplicationStatus);
+	void OnNetworkStatusChanged(EOS_ENetworkStatus NetworkStatus);
+	EOS_EApplicationStatus CachedApplicationStatus = EOS_EApplicationStatus::EOS_AS_Foreground;
+	EOS_ENetworkStatus CachedNetworkStatus = EOS_ENetworkStatus::EOS_NS_Online;
 
 	friend struct FEOSPlatformHandle;
 
@@ -74,7 +80,6 @@ protected:
 	void ReleasePlatform(EOS_HPlatform PlatformHandle);
 	void ReleaseReleasedPlatforms();
 	void SetupTicker();
-	bool Tick(float);
 	void OnLogVerbosityChanged(const FLogCategoryName& CategoryName, ELogVerbosity::Type OldVerbosity, ELogVerbosity::Type NewVerbosity);
 
 #if EOSSDK_RUNTIME_LOAD_REQUIRED
