- Fix dereferencing TCHAR* in log macro

diff --git a/Engine/Plugins/Online/EOSShared/Source/EOSShared/Private/EOSSDKManager.cpp b/Engine/Plugins/Online/EOSShared/Source/EOSShared/Private/EOSSDKManager.cpp
--- a/Engine/Plugins/Online/EOSShared/Source/EOSShared/Private/EOSSDKManager.cpp
+++ b/Engine/Plugins/Online/EOSShared/Source/EOSShared/Private/EOSSDKManager.cpp
@@ -524,7 +524,7 @@ bool FEOSSDKManager::Tick(float)
 
 void FEOSSDKManager::OnApplicationStatusChanged(EOS_EApplicationStatus ApplicationStatus)
 {
-	UE_LOG(LogEOSSDK, Log, TEXT("OnApplicationStatusChanged [%s] -> [%s]"), *LexToString(CachedApplicationStatus), *LexToString(ApplicationStatus));
+	UE_LOG(LogEOSSDK, Log, TEXT("OnApplicationStatusChanged [%s] -> [%s]"), LexToString(CachedApplicationStatus), LexToString(ApplicationStatus));
 	CachedApplicationStatus = ApplicationStatus;
 	for (EOS_HPlatform PlatformHandle : ActivePlatforms)
 	{
@@ -534,7 +534,7 @@ void FEOSSDKManager::OnApplicationStatusChanged(EOS_EApplicationStatus Applicati
 
 void FEOSSDKManager::OnNetworkStatusChanged(EOS_ENetworkStatus NetworkStatus)
 {
-	UE_LOG(LogEOSSDK, Log, TEXT("OnNetworkStatusChanged [%s] -> [%s]"), *LexToString(CachedNetworkStatus), *LexToString(NetworkStatus));
+	UE_LOG(LogEOSSDK, Log, TEXT("OnNetworkStatusChanged [%s] -> [%s]"), LexToString(CachedNetworkStatus), LexToString(NetworkStatus));
 	CachedNetworkStatus = NetworkStatus;
 	for (EOS_HPlatform PlatformHandle : ActivePlatforms)
 	{
