Add platform instance support to EOSShared

diff --git a/Engine/Plugins/Online/EOSShared/Source/EOSShared/Private/EOSSDKManager.cpp b/Engine/Plugins/Online/EOSShared/Source/EOSShared/Private/EOSSDKManager.cpp
--- a/Engine/Plugins/Online/EOSShared/Source/EOSShared/Private/EOSSDKManager.cpp
+++ b/Engine/Plugins/Online/EOSShared/Source/EOSShared/Private/EOSSDKManager.cpp
@@ -345,7 +345,7 @@ void FEOSSDKManager::SetDefaultPlatformConfigName(const FString& PlatformConfigN
 	}
 }
 
-IEOSPlatformHandlePtr FEOSSDKManager::CreatePlatform(const FString& PlatformConfigName)
+IEOSPlatformHandlePtr FEOSSDKManager::CreatePlatform(const FString& PlatformConfigName, FName InstanceName)
 {
 	if (PlatformConfigName.IsEmpty())
 	{
@@ -353,23 +353,31 @@ IEOSPlatformHandlePtr FEOSSDKManager::CreatePlatform(const FString& PlatformConf
 		return IEOSPlatformHandlePtr();
 	}
 
-	IEOSPlatformHandleWeakPtr* WeakPlatformHandle = PlatformHandles.Find(PlatformConfigName);
-	if (WeakPlatformHandle)
+	const FEOSSDKPlatformConfig* const PlatformConfig = GetPlatformConfig(PlatformConfigName, true);
+	if (!PlatformConfig)
+	{
+		return IEOSPlatformHandlePtr();
+	}
+
+	TMap<FName, IEOSPlatformHandleWeakPtr>* PlatformMap = PlatformHandles.Find(PlatformConfigName);
+	if (PlatformMap)
 	{
-		if (WeakPlatformHandle->IsValid())
+		IEOSPlatformHandleWeakPtr* WeakPlatformHandle = PlatformMap->Find(InstanceName);
+		if (WeakPlatformHandle)
 		{
-			UE_LOG(LogEOSSDK, Verbose, TEXT("Found existing platform handle: %s"), *PlatformConfigName);
-			return WeakPlatformHandle->Pin();
-		}
+			if (IEOSPlatformHandlePtr Pinned = WeakPlatformHandle->Pin())
+			{
+				UE_LOG(LogEOSSDK, Verbose, TEXT("Found existing platform handle: PlatformConfigName=%s InstanceName=%s"), *PlatformConfigName, *InstanceName.ToString());
+				return Pinned;
+			}
 
-		UE_LOG(LogEOSSDK, Verbose, TEXT("Removing stale platform handle pointer: %s"), *PlatformConfigName);
-		PlatformHandles.Remove(PlatformConfigName);
+			UE_LOG(LogEOSSDK, Verbose, TEXT("Removing stale platform handle pointer: PlatformConfigName=%s InstanceName=%s"), *PlatformConfigName, *InstanceName.ToString());
+			PlatformMap->Remove(InstanceName);
+		}
 	}
-
-	const FEOSSDKPlatformConfig* const PlatformConfig = GetPlatformConfig(PlatformConfigName, true);
-	if (!PlatformConfig)
+	else
 	{
-		return IEOSPlatformHandlePtr();
+		PlatformMap = &PlatformHandles.Emplace(PlatformConfigName);
 	}
 
 	const FTCHARToUTF8 Utf8ProductId(*PlatformConfig->ProductId);
@@ -415,8 +423,8 @@ IEOSPlatformHandlePtr FEOSSDKManager::CreatePlatform(const FString& PlatformConf
 	IEOSPlatformHandlePtr PlatformHandle = CreatePlatform(*PlatformConfig, PlatformOptions);
 	if (PlatformHandle.IsValid())
 	{
-		UE_LOG(LogEOSSDK, Verbose, TEXT("Created platform handle: %s"), *PlatformConfigName);
-		PlatformHandles.Emplace(PlatformConfigName, PlatformHandle);
+		UE_LOG(LogEOSSDK, Verbose, TEXT("Created platform handle: PlatformConfigName=%s InstanceName=%s"), *PlatformConfigName, *InstanceName.ToString());
+		PlatformMap->Emplace(InstanceName, PlatformHandle);
 	}
 
 	return PlatformHandle;
diff --git a/Engine/Plugins/Online/EOSShared/Source/EOSShared/Private/EOSSDKManager.h b/Engine/Plugins/Online/EOSShared/Source/EOSShared/Private/EOSSDKManager.h
--- a/Engine/Plugins/Online/EOSShared/Source/EOSShared/Private/EOSSDKManager.h
+++ b/Engine/Plugins/Online/EOSShared/Source/EOSShared/Private/EOSSDKManager.h
@@ -37,7 +37,7 @@ public:
 	virtual const FString& GetDefaultPlatformConfigName() override;
 	virtual void SetDefaultPlatformConfigName(const FString& PlatformConfigName) override;
 
-	virtual IEOSPlatformHandlePtr CreatePlatform(const FString& PlatformConfigName) override;
+	virtual IEOSPlatformHandlePtr CreatePlatform(const FString& PlatformConfigName, FName InstanceName = NAME_None) override;
 	virtual IEOSPlatformHandlePtr CreatePlatform(EOS_Platform_Options& PlatformOptions) override;
 
 	virtual FString GetProductName() const override;
@@ -103,7 +103,7 @@ protected:
 	/** Default platform config name to use. */
 	FString DefaultPlatformConfigName;
 	/** Cache of named platform handles that have been created. */
-	TMap<FString, IEOSPlatformHandleWeakPtr> PlatformHandles;
+	TMap<FString, TMap<FName, IEOSPlatformHandleWeakPtr>> PlatformHandles;
 
 	// Config
 	/** Interval between platform ticks. 0 means we tick every frame. */
diff --git a/Engine/Plugins/Online/EOSShared/Source/EOSShared/Public/IEOSSDKManager.h b/Engine/Plugins/Online/EOSShared/Source/EOSShared/Public/IEOSSDKManager.h
--- a/Engine/Plugins/Online/EOSShared/Source/EOSShared/Public/IEOSSDKManager.h
+++ b/Engine/Plugins/Online/EOSShared/Source/EOSShared/Public/IEOSSDKManager.h
@@ -101,7 +101,7 @@ public:
 	 * Create a platform handle for a platform config name. Config is loaded from .ini files if it was not added with AddPlatformConfig.
 	 * If a platform handle already exists for the config name, this will return a shared pointer to that handle and not create a new one.
 	 */
-	virtual IEOSPlatformHandlePtr CreatePlatform(const FString& PlatformConfigName) = 0;
+	virtual IEOSPlatformHandlePtr CreatePlatform(const FString& PlatformConfigName, FName InstanceName = NAME_None) = 0;
 
 	/** Create a platform handle using EOSSDK options directly. */
 	virtual IEOSPlatformHandlePtr CreatePlatform(EOS_Platform_Options& PlatformOptions) = 0;
