Addition of logic to forward EOSSDK Network and application status events

diff --git a/Engine/Plugins/Online/EOSShared/Source/EOSShared/Private/EOSSDKManager.cpp b/Engine/Plugins/Online/EOSShared/Source/EOSShared/Private/EOSSDKManager.cpp
--- a/Engine/Plugins/Online/EOSShared/Source/EOSShared/Private/EOSSDKManager.cpp
+++ b/Engine/Plugins/Online/EOSShared/Source/EOSShared/Private/EOSSDKManager.cpp
@@ -244,6 +244,8 @@ EOS_EResult FEOSSDKManager::Initialize()
 				UE_LOG(LogEOSSDK, Warning, TEXT("EOS_Logging_SetLogLevel failed Verbosity=%s error=[%s]"), ToString(LogEOSSDK.GetVerbosity()), *LexToString(EosResult));
 			}
 #endif // !NO_LOGGING
+
+			FCoreDelegates::OnNetworkConnectionStatusChanged.AddRaw(this, &FEOSSDKManager::OnNetworkConnectionStatusChanged);
 		}
 		else
 		{
@@ -452,7 +454,7 @@ IEOSPlatformHandlePtr FEOSSDKManager::CreatePlatform(EOS_Platform_Options& Platf
 			SetupTicker();
 
 			EOS_Platform_SetApplicationStatus(PlatformHandle, CachedApplicationStatus);
-			EOS_Platform_SetNetworkStatus(PlatformHandle, CachedNetworkStatus);
+			EOS_Platform_SetNetworkStatus(PlatformHandle, ConvertNetworkStatus(FPlatformMisc::GetNetworkConnectionStatus()));
 
 			// Tick the platform once to work around EOSSDK error logging that occurs if you create then immediately destroy a platform.
 			SharedPlatform->Tick();
@@ -530,23 +532,40 @@ bool FEOSSDKManager::Tick(float)
 	return true;
 }
 
-void FEOSSDKManager::OnApplicationStatusChanged(EOS_EApplicationStatus ApplicationStatus)
+EOS_ENetworkStatus FEOSSDKManager::ConvertNetworkStatus(ENetworkConnectionStatus Status)
 {
-	UE_LOG(LogEOSSDK, Log, TEXT("OnApplicationStatusChanged [%s] -> [%s]"), LexToString(CachedApplicationStatus), LexToString(ApplicationStatus));
-	CachedApplicationStatus = ApplicationStatus;
+	switch (Status)
+	{
+	case ENetworkConnectionStatus::Unknown:			return EOS_ENetworkStatus::EOS_NS_Online;
+	case ENetworkConnectionStatus::Disabled:		return EOS_ENetworkStatus::EOS_NS_Disabled;
+	case ENetworkConnectionStatus::Local:			return EOS_ENetworkStatus::EOS_NS_Offline;
+	case ENetworkConnectionStatus::Connected:		return EOS_ENetworkStatus::EOS_NS_Online;
+	}
+
+	checkNoEntry();
+	return EOS_ENetworkStatus::EOS_NS_Disabled;
+}
+
+void FEOSSDKManager::OnNetworkConnectionStatusChanged(ENetworkConnectionStatus LastConnectionState, ENetworkConnectionStatus ConnectionState)
+{
+	const EOS_ENetworkStatus OldNetworkStatus = ConvertNetworkStatus(LastConnectionState);
+	const EOS_ENetworkStatus NewNetworkStatus = ConvertNetworkStatus(ConnectionState);
+
+	UE_LOG(LogEOSSDK, Log, TEXT("OnNetworkConnectionStatusChanged [%s] -> [%s]"), LexToString(OldNetworkStatus), LexToString(NewNetworkStatus));
+
 	for (EOS_HPlatform PlatformHandle : ActivePlatforms)
 	{
-		EOS_Platform_SetApplicationStatus(PlatformHandle, ApplicationStatus);
+		EOS_Platform_SetNetworkStatus(PlatformHandle, NewNetworkStatus);
 	}
 }
 
-void FEOSSDKManager::OnNetworkStatusChanged(EOS_ENetworkStatus NetworkStatus)
+void FEOSSDKManager::OnApplicationStatusChanged(EOS_EApplicationStatus ApplicationStatus)
 {
-	UE_LOG(LogEOSSDK, Log, TEXT("OnNetworkStatusChanged [%s] -> [%s]"), LexToString(CachedNetworkStatus), LexToString(NetworkStatus));
-	CachedNetworkStatus = NetworkStatus;
+	UE_LOG(LogEOSSDK, Log, TEXT("OnApplicationStatusChanged [%s] -> [%s]"), LexToString(CachedApplicationStatus), LexToString(ApplicationStatus));
+	CachedApplicationStatus = ApplicationStatus;
 	for (EOS_HPlatform PlatformHandle : ActivePlatforms)
 	{
-		EOS_Platform_SetNetworkStatus(PlatformHandle, NetworkStatus);
+		EOS_Platform_SetApplicationStatus(PlatformHandle, ApplicationStatus);
 	}
 }
 
@@ -665,6 +684,8 @@ void FEOSSDKManager::Shutdown()
 		CallbackObjects.Empty();
 
 		bInitialized = false;
+
+		FCoreDelegates::OnNetworkConnectionStatusChanged.RemoveAll(this);
 	}
 }
 
diff --git a/Engine/Plugins/Online/EOSShared/Source/EOSShared/Private/EOSSDKManager.h b/Engine/Plugins/Online/EOSShared/Source/EOSShared/Private/EOSSDKManager.h
--- a/Engine/Plugins/Online/EOSShared/Source/EOSShared/Private/EOSSDKManager.h
+++ b/Engine/Plugins/Online/EOSShared/Source/EOSShared/Private/EOSSDKManager.h
@@ -68,10 +68,10 @@ protected:
 	virtual IEOSPlatformHandlePtr CreatePlatform(const FEOSSDKPlatformConfig& PlatformConfig, EOS_Platform_Options& PlatformOptions);
 	virtual bool Tick(float);
 
+	static EOS_ENetworkStatus ConvertNetworkStatus(ENetworkConnectionStatus Status);
+	void OnNetworkConnectionStatusChanged(ENetworkConnectionStatus LastConnectionState, ENetworkConnectionStatus ConnectionState);
 	void OnApplicationStatusChanged(EOS_EApplicationStatus ApplicationStatus);
-	void OnNetworkStatusChanged(EOS_ENetworkStatus NetworkStatus);
 	EOS_EApplicationStatus CachedApplicationStatus = EOS_EApplicationStatus::EOS_AS_Foreground;
-	EOS_ENetworkStatus CachedNetworkStatus = EOS_ENetworkStatus::EOS_NS_Online;
 
 	friend struct FEOSPlatformHandle;
 
