Add LexFromString for EOS_EpicAccountId.

diff --git a/Engine/Plugins/Online/EOSShared/Source/EOSShared/Private/EOSShared.cpp b/Engine/Plugins/Online/EOSShared/Source/EOSShared/Private/EOSShared.cpp
--- a/Engine/Plugins/Online/EOSShared/Source/EOSShared/Private/EOSShared.cpp
+++ b/Engine/Plugins/Online/EOSShared/Source/EOSShared/Private/EOSShared.cpp
@@ -52,6 +52,11 @@ FString LexToString(const EOS_EpicAccountId AccountId)
 	return Result;
 }
 
+void LexFromString(EOS_EpicAccountId& AccountId, const TCHAR* String)
+{
+	AccountId = EOS_EpicAccountId_FromString(TCHAR_TO_UTF8(String));
+}
+
 const TCHAR* LexToString(const EOS_EApplicationStatus ApplicationStatus)
 {
 	switch (ApplicationStatus)
diff --git a/Engine/Plugins/Online/EOSShared/Source/EOSShared/Public/EOSShared.h b/Engine/Plugins/Online/EOSShared/Source/EOSShared/Public/EOSShared.h
--- a/Engine/Plugins/Online/EOSShared/Source/EOSShared/Public/EOSShared.h
+++ b/Engine/Plugins/Online/EOSShared/Source/EOSShared/Public/EOSShared.h
@@ -44,6 +44,7 @@ inline EOS_ProductUserId EOSProductUserIdFromString(const TCHAR* String)
 }
 
 EOSSHARED_API FString LexToString(const EOS_EpicAccountId AccountId);
+EOSSHARED_API void LexFromString(EOS_EpicAccountId& AccountId, const TCHAR* String);
 
 EOSSHARED_API const TCHAR* LexToString(const EOS_EApplicationStatus ApplicationStatus);
 EOSSHARED_API const TCHAR* LexToString(const EOS_EAuthTokenType AuthTokenType);
