- Make AuthScopeFlags configurable
- Make LexFromString for AuthScopeFlags support multiple flags in one string.

diff --git a/Engine/Config/BaseEngine.ini b/Engine/Config/BaseEngine.ini
--- a/Engine/Config/BaseEngine.ini
+++ b/Engine/Config/BaseEngine.ini
@@ -2067,6 +2067,11 @@ bIncludeSamsungLocText=true
 [OnlineSubsystemApple]
 bEnabled=false
 
+[/Script/OnlineSubsystemEOS.EOSSettings]
++AuthScopeFlags="BasicProfile"
++AuthScopeFlags="FriendsList"
++AuthScopeFlags="Presence"
+
 [OnlineSubsystemNull]
 bEnabled=true
 bAutoLoginAtStartup=true
diff --git a/Engine/Plugins/Online/EOSShared/Source/EOSShared/Private/EOSShared.cpp b/Engine/Plugins/Online/EOSShared/Source/EOSShared/Private/EOSShared.cpp
--- a/Engine/Plugins/Online/EOSShared/Source/EOSShared/Private/EOSShared.cpp
+++ b/Engine/Plugins/Online/EOSShared/Source/EOSShared/Private/EOSShared.cpp
@@ -3,6 +3,8 @@
 #if WITH_EOS_SDK
 
 #include "EOSShared.h"
+
+#include "String/ParseTokens.h"
 #include "EOSSharedTypes.h"
 
 #include "eos_auth_types.h"
@@ -206,37 +208,48 @@ bool LexFromString(EOS_EExternalCredentialType& OutEnum, const TCHAR* InString)
 	return true;
 }
 
-bool LexFromString(EOS_EAuthScopeFlags& OutEnum, const TCHAR* InString)
+bool LexFromString(EOS_EAuthScopeFlags& OutEnum, const FStringView InString)
 {
-	if (FCString::Stricmp(InString, TEXT("BasicProfile")) == 0)
-	{
-		OutEnum = EOS_EAuthScopeFlags::EOS_AS_BasicProfile;
-	}
-	else if (FCString::Stricmp(InString, TEXT("FriendsList")) == 0)
-	{
-		OutEnum = EOS_EAuthScopeFlags::EOS_AS_FriendsList;
-	}
-	else if (FCString::Stricmp(InString, TEXT("Presence")) == 0)
-	{
-		OutEnum = EOS_EAuthScopeFlags::EOS_AS_Presence;
-	}
-	else if (FCString::Stricmp(InString, TEXT("FriendsManagement")) == 0)
-	{
-		OutEnum = EOS_EAuthScopeFlags::EOS_AS_FriendsManagement;
-	}
-	else if (FCString::Stricmp(InString, TEXT("Email")) == 0)
-	{
-		OutEnum = EOS_EAuthScopeFlags::EOS_AS_Email;
-	}
-	else if (FCString::Stricmp(InString, TEXT("NoFlags")) == 0 || FCString::Stricmp(InString, TEXT("None")) == 0)
-	{
-		OutEnum = EOS_EAuthScopeFlags::EOS_AS_NoFlags;
-	}
-	else
-	{
-		return false;
-	}
-	return true;
+	OutEnum = EOS_EAuthScopeFlags::EOS_AS_NoFlags;
+	bool bParsedOk = true;
+
+	using namespace UE::String;
+	const EParseTokensOptions ParseOptions = EParseTokensOptions::SkipEmpty | EParseTokensOptions::Trim;
+	TFunctionRef<void(FStringView)> ParseFunc = [&OutEnum, &bParsedOk](FStringView Token)
+	{
+		if (Token == TEXT("BasicProfile"))
+		{
+			OutEnum |= EOS_EAuthScopeFlags::EOS_AS_BasicProfile;
+		}
+		else if (Token == TEXT("FriendsList"))
+		{
+			OutEnum |= EOS_EAuthScopeFlags::EOS_AS_FriendsList;
+		}
+		else if (Token == TEXT("Presence"))
+		{
+			OutEnum |= EOS_EAuthScopeFlags::EOS_AS_Presence;
+		}
+		else if (Token == TEXT("FriendsManagement"))
+		{
+			OutEnum |= EOS_EAuthScopeFlags::EOS_AS_FriendsManagement;
+		}
+		else if (Token == TEXT("Email"))
+		{
+			OutEnum |= EOS_EAuthScopeFlags::EOS_AS_Email;
+		}
+		else if (Token == TEXT("Country"))
+		{
+			OutEnum |= EOS_EAuthScopeFlags::EOS_AS_Country;
+		}
+		else
+		{
+			checkNoEntry();
+			bParsedOk = false;
+		}
+	};
+	ParseTokens(InString, TCHAR('|'), ParseFunc, ParseOptions);
+
+	return bParsedOk;
 }
 
 bool LexFromString(EOS_ELoginCredentialType& OutEnum, const TCHAR* InString)
diff --git a/Engine/Plugins/Online/EOSShared/Source/EOSShared/Public/EOSShared.h b/Engine/Plugins/Online/EOSShared/Source/EOSShared/Public/EOSShared.h
--- a/Engine/Plugins/Online/EOSShared/Source/EOSShared/Public/EOSShared.h
+++ b/Engine/Plugins/Online/EOSShared/Source/EOSShared/Public/EOSShared.h
@@ -55,7 +55,7 @@ EOSSHARED_API const TCHAR* LexToString(const EOS_ELoginStatus LoginStatus);
 EOSSHARED_API const TCHAR* LexToString(const EOS_ENetworkStatus NetworkStatus);
 EOSSHARED_API const TCHAR* LexToString(const EOS_Presence_EStatus PresenceStatus);
 
-EOSSHARED_API bool LexFromString(EOS_EAuthScopeFlags& OutEnum, const TCHAR* InString);
+EOSSHARED_API bool LexFromString(EOS_EAuthScopeFlags& OutEnum, const FStringView InString);
 EOSSHARED_API bool LexFromString(EOS_EExternalCredentialType& OutEnum, const TCHAR* InString);
 EOSSHARED_API bool LexFromString(EOS_ELoginCredentialType& OutEnum, const TCHAR* InString);
 
diff --git a/Engine/Plugins/Online/OnlineServicesEOSGS/Source/Private/Online/AuthEOSGS.cpp b/Engine/Plugins/Online/OnlineServicesEOSGS/Source/Private/Online/AuthEOSGS.cpp
--- a/Engine/Plugins/Online/OnlineServicesEOSGS/Source/Private/Online/AuthEOSGS.cpp
+++ b/Engine/Plugins/Online/OnlineServicesEOSGS/Source/Private/Online/AuthEOSGS.cpp
@@ -236,7 +236,6 @@ private:
 
 	static const FEOSAuthTranslationTraits* GetLoginTranslatorTraits(FName Name);
 	static const FEOSExternalAuthTranslationTraits* GetExternalAuthTranslationTraits(FName ExternalAuthType);
-	static const EOS_EAuthScopeFlags* GetScopeFlag(const FString& ScopeName);
 
 	EOS_Auth_Credentials CredentialsData;
 	TArray<char> IdUtf8;
@@ -284,19 +283,6 @@ const FEOSExternalAuthTranslationTraits* FEOSAuthLoginOptions::GetExternalAuthTr
 	return SupportedExternalAuthTraits.Find(ExternalAuthType);
 }
 
-const EOS_EAuthScopeFlags* FEOSAuthLoginOptions::GetScopeFlag(const FString& ScopeName)
-{
-	static const TMap<FString, EOS_EAuthScopeFlags> SupportedScopes = {
-		{ TEXT("BasicProfile"), EOS_EAuthScopeFlags::EOS_AS_BasicProfile },
-		{ TEXT("FriendsList"), EOS_EAuthScopeFlags::EOS_AS_FriendsList },
-		{ TEXT("Presence"), EOS_EAuthScopeFlags::EOS_AS_Presence },
-		{ TEXT("FriendsManagement"), EOS_EAuthScopeFlags::EOS_AS_FriendsManagement },
-		{ TEXT("Email"), EOS_EAuthScopeFlags::EOS_AS_Email },
-	};
-
-	return SupportedScopes.Find(ScopeName);
-}
-
 FEOSAuthLoginOptions::FEOSAuthLoginOptions(FEOSAuthLoginOptions&& Other)
 {
 	*this = MoveTemp(Other);
@@ -448,9 +434,10 @@ TDefaultErrorResultInternal<FEOSAuthLoginOptions> FEOSAuthLoginOptions::CreateIm
 	bool bAllScopesValid = true;
 	for (const FString& Scope : Scopes)
 	{
-		if (const EOS_EAuthScopeFlags* ScopeFlag = GetScopeFlag(Scope))
+		EOS_EAuthScopeFlags ScopeFlags;
+		if (LexFromString(ScopeFlags, Scope))
 		{
-			EOSAuthLoginOptions.ScopeFlags |= *ScopeFlag;
+			EOSAuthLoginOptions.ScopeFlags |= ScopeFlags;
 		}
 		else
 		{
diff --git a/Engine/Plugins/Online/OnlineSubsystemEOS/Source/OnlineSubsystemEOS/Private/EOSSettings.cpp b/Engine/Plugins/Online/OnlineSubsystemEOS/Source/OnlineSubsystemEOS/Private/EOSSettings.cpp
--- a/Engine/Plugins/Online/OnlineSubsystemEOS/Source/OnlineSubsystemEOS/Private/EOSSettings.cpp
+++ b/Engine/Plugins/Online/OnlineSubsystemEOS/Source/OnlineSubsystemEOS/Private/EOSSettings.cpp
@@ -163,6 +163,7 @@ const FEOSSettings& UEOSSettings::ManualGetSettings()
 		GConfig->GetBool(INI_SECTION, TEXT("bMirrorPresenceToEAS"), CachedSettings->bMirrorPresenceToEAS, GEngineIni);
 		// Artifacts explicitly skipped
 		GConfig->GetArray(INI_SECTION, TEXT("TitleStorageTags"), CachedSettings->TitleStorageTags, GEngineIni);
+		GConfig->GetArray(INI_SECTION, TEXT("AuthScopeFlags"), CachedSettings->AuthScopeFlags, GEngineIni);
 	}
 
 	return *CachedSettings;
@@ -188,6 +189,7 @@ FEOSSettings UEOSSettings::ToNative() const
 	Native.bMirrorPresenceToEAS = bMirrorPresenceToEAS;
 	Algo::Transform(Artifacts, Native.Artifacts, &FArtifactSettings::ToNative);
 	Native.TitleStorageTags = TitleStorageTags;
+	Native.AuthScopeFlags = AuthScopeFlags;
 
 	return Native;
 }
diff --git a/Engine/Plugins/Online/OnlineSubsystemEOS/Source/OnlineSubsystemEOS/Private/UserManagerEOS.cpp b/Engine/Plugins/Online/OnlineSubsystemEOS/Source/OnlineSubsystemEOS/Private/UserManagerEOS.cpp
--- a/Engine/Plugins/Online/OnlineSubsystemEOS/Source/OnlineSubsystemEOS/Private/UserManagerEOS.cpp
+++ b/Engine/Plugins/Online/OnlineSubsystemEOS/Source/OnlineSubsystemEOS/Private/UserManagerEOS.cpp
@@ -429,7 +429,16 @@ bool FUserManagerEOS::Login(int32 LocalUserNum, const FOnlineAccountCredentials&
 
 	EOS_Auth_LoginOptions LoginOptions = { };
 	LoginOptions.ApiVersion = EOS_AUTH_LOGIN_API_LATEST;
-	LoginOptions.ScopeFlags = EOS_EAuthScopeFlags::EOS_AS_BasicProfile | EOS_EAuthScopeFlags::EOS_AS_FriendsList | EOS_EAuthScopeFlags::EOS_AS_Presence;
+
+	LoginOptions.ScopeFlags = EOS_EAuthScopeFlags::EOS_AS_NoFlags;
+	for (const FString& FlagsStr : Settings.AuthScopeFlags)
+	{
+		EOS_EAuthScopeFlags Flags;
+		if (LexFromString(Flags, FlagsStr))
+		{
+			LoginOptions.ScopeFlags |= Flags;
+		}
+	}
 
 	FPlatformEOSHelpersPtr EOSHelpers = EOSSubsystem->GetEOSHelpers();
 
diff --git a/Engine/Plugins/Online/OnlineSubsystemEOS/Source/OnlineSubsystemEOS/Public/EOSSettings.h b/Engine/Plugins/Online/OnlineSubsystemEOS/Source/OnlineSubsystemEOS/Public/EOSSettings.h
--- a/Engine/Plugins/Online/OnlineSubsystemEOS/Source/OnlineSubsystemEOS/Public/EOSSettings.h
+++ b/Engine/Plugins/Online/OnlineSubsystemEOS/Source/OnlineSubsystemEOS/Public/EOSSettings.h
@@ -83,6 +83,7 @@ struct FEOSSettings
 	bool bMirrorPresenceToEAS;
 	TArray<FEOSArtifactSettings> Artifacts;
 	TArray<FString> TitleStorageTags;
+	TArray<FString> AuthScopeFlags;
 };
 
 UCLASS(Config=Engine, DefaultConfig)
@@ -136,6 +137,10 @@ public:
 	UPROPERTY(Config, EditAnywhere, BlueprintReadOnly, Category="EOS Settings")
 	TArray<FArtifactSettings> Artifacts;
 
+	/** Auth scopes to request during login */
+	UPROPERTY(Config, EditAnywhere, BlueprintReadOnly, Category = "EOS Settings")
+	TArray<FString> AuthScopeFlags;
+
 	/** Set to true to have Epic Accounts used (friends list will be unified with the default platform) */
 	UPROPERTY(Config, EditAnywhere, BlueprintReadOnly, Category="EOSPlus Login Settings", DisplayName="Use Epic Account for EOS login (requires account linking)")
 	bool bUseEAS = false;
