- Add log filtering for EOSSDK, to work around some spurious errors/warnings from the SDK.

diff --git a/Engine/Plugins/Online/EOSShared/Source/EOSShared/Private/EOSSDKManager.cpp b/Engine/Plugins/Online/EOSShared/Source/EOSShared/Private/EOSSDKManager.cpp
--- a/Engine/Plugins/Online/EOSShared/Source/EOSShared/Private/EOSSDKManager.cpp
+++ b/Engine/Plugins/Online/EOSShared/Source/EOSShared/Private/EOSSDKManager.cpp
@@ -71,18 +71,41 @@ namespace
 	}
 
 #if !NO_LOGGING
+	bool IsLogLevelSuppressable(EOS_ELogLevel Level)
+	{
+		return Level > EOS_ELogLevel::EOS_LOG_Off && Level < EOS_ELogLevel::EOS_LOG_Info;
+	}
+
 	void EOS_CALL EOSLogMessageReceived(const EOS_LogMessage* Message)
 	{
-#define EOSLOG(Level) UE_LOG(LogEOSSDK, Level, TEXT("%s: %s"), UTF8_TO_TCHAR(Message->Category), *MessageStr)
+#define EOSLOG(Level) UE_LOG(LogEOSSDK, Level, TEXT("%hs: %s"), Message->Category, *MessageStr)
+#define EOSLOG_SUPPRESS(Level) if(bSuppressLogLevel) { EOSLOG(Log); } else { EOSLOG(Level); }
 
 		FString MessageStr(UTF8_TO_TCHAR(Message->Message));
 		MessageStr.TrimStartAndEndInline();
 
+		// Check if this log line is suppressed.
+		bool bSuppressLogLevel = false;
+		if (IsLogLevelSuppressable(Message->Level))
+		{
+			if (FEOSSDKManager* Manager = static_cast<FEOSSDKManager*>(IEOSSDKManager::Get()))
+			{
+				for (const FString& SuppressedString : Manager->GetSuppressedLogStrings())
+				{
+					if (MessageStr.Contains(SuppressedString))
+					{
+						bSuppressLogLevel = true;
+						break;
+					}
+				}
+			}
+		}
+
 		switch (Message->Level)
 		{
-		case EOS_ELogLevel::EOS_LOG_Fatal:			EOSLOG(Fatal); break;
-		case EOS_ELogLevel::EOS_LOG_Error:			EOSLOG(Error); break;
-		case EOS_ELogLevel::EOS_LOG_Warning:		EOSLOG(Warning); break;
+		case EOS_ELogLevel::EOS_LOG_Fatal:			EOSLOG_SUPPRESS(Fatal); break;
+		case EOS_ELogLevel::EOS_LOG_Error:			EOSLOG_SUPPRESS(Error); break;
+		case EOS_ELogLevel::EOS_LOG_Warning:		EOSLOG_SUPPRESS(Warning); break;
 		case EOS_ELogLevel::EOS_LOG_Info:			EOSLOG(Log); break;
 		case EOS_ELogLevel::EOS_LOG_Verbose:		EOSLOG(Verbose); break;
 		case EOS_ELogLevel::EOS_LOG_VeryVerbose:	EOSLOG(VeryVerbose); break;
@@ -92,6 +115,7 @@ namespace
 			break;
 		}
 #undef EOSLOG
+#undef EOSLOG_SUPPRESS
 	}
 
 	EOS_ELogLevel ConvertLogLevel(ELogVerbosity::Type LogLevel)
@@ -496,8 +520,11 @@ void FEOSSDKManager::OnConfigSectionsChanged(const FString& IniFilename, const T
 
 void FEOSSDKManager::LoadConfig()
 {
+	const TCHAR* SectionName = TEXT("EOSSDK");
 	ConfigTickIntervalSeconds = 0.f;
-	GConfig->GetDouble(TEXT("EOSSDK"), TEXT("TickIntervalSeconds"), ConfigTickIntervalSeconds, GEngineIni);
+	GConfig->GetDouble(SectionName, TEXT("TickIntervalSeconds"), ConfigTickIntervalSeconds, GEngineIni);
+
+	GConfig->GetArray(SectionName, TEXT("SuppressedLogStrings"), SuppressedLogStrings, GEngineIni);
 
 	SetupTicker();
 }
diff --git a/Engine/Plugins/Online/EOSShared/Source/EOSShared/Private/EOSSDKManager.h b/Engine/Plugins/Online/EOSShared/Source/EOSShared/Private/EOSSDKManager.h
--- a/Engine/Plugins/Online/EOSShared/Source/EOSShared/Private/EOSSDKManager.h
+++ b/Engine/Plugins/Online/EOSShared/Source/EOSShared/Private/EOSSDKManager.h
@@ -108,6 +108,11 @@ protected:
 	// Config
 	/** Interval between platform ticks. 0 means we tick every frame. */
 	double ConfigTickIntervalSeconds = 0.f;
+public:
+	const TArray<FString>& GetSuppressedLogStrings() const { return SuppressedLogStrings; }
+protected:
+	/** Log strings that should be downgraded from Error/Warning to Log. */
+	TArray<FString> SuppressedLogStrings;
 };
 
 struct FEOSPlatformHandle : public IEOSPlatformHandle
