- revert change to %hs which apparently does not work correctly with UTF8 characters.

diff --git a/Engine/Plugins/Online/EOSShared/Source/EOSShared/Private/EOSSDKManager.cpp b/Engine/Plugins/Online/EOSShared/Source/EOSShared/Private/EOSSDKManager.cpp
--- a/Engine/Plugins/Online/EOSShared/Source/EOSShared/Private/EOSSDKManager.cpp
+++ b/Engine/Plugins/Online/EOSShared/Source/EOSShared/Private/EOSSDKManager.cpp
@@ -78,7 +78,7 @@ namespace
 
 	void EOS_CALL EOSLogMessageReceived(const EOS_LogMessage* Message)
 	{
-#define EOSLOG(Level) UE_LOG(LogEOSSDK, Level, TEXT("%hs: %s"), Message->Category, *MessageStr)
+#define EOSLOG(Level) UE_LOG(LogEOSSDK, Level, TEXT("%s: %s"), UTF8_TO_TCHAR(Message->Category), *MessageStr)
 #define EOSLOG_SUPPRESS(Level) if(bSuppressLogLevel) { EOSLOG(Log); } else { EOSLOG(Level); }
 
 		FString MessageStr(UTF8_TO_TCHAR(Message->Message));
