Updated EOS application status when moving to background/foreground on Android and iOS

diff --git a/Engine/Plugins/Online/EOSShared/Source/EOSShared/Private/Android/AndroidEOSSDKManager.cpp b/Engine/Plugins/Online/EOSShared/Source/EOSShared/Private/Android/AndroidEOSSDKManager.cpp
--- a/Engine/Plugins/Online/EOSShared/Source/EOSShared/Private/Android/AndroidEOSSDKManager.cpp
+++ b/Engine/Plugins/Online/EOSShared/Source/EOSShared/Private/Android/AndroidEOSSDKManager.cpp
@@ -3,12 +3,24 @@
 #include "AndroidEOSSDKManager.h"
 
 #if WITH_EOS_SDK
-
-#include "HAL/FileManager.h"
 #include "EOSShared.h"
+#include "HAL/FileManager.h"
+#include "Misc/CoreDelegates.h"
 
 #include "eos_android.h"
 
+FAndroidEOSSDKManager::FAndroidEOSSDKManager()
+{
+	FCoreDelegates::ApplicationHasEnteredForegroundDelegate.AddRaw(this, &FAndroidEOSSDKManager::OnApplicationStatusChanged, EOS_EApplicationStatus::EOS_AS_Foreground);
+	FCoreDelegates::ApplicationWillEnterBackgroundDelegate.AddRaw(this, &FAndroidEOSSDKManager::OnApplicationStatusChanged, EOS_EApplicationStatus::EOS_AS_BackgroundSuspended);
+}
+
+FAndroidEOSSDKManager::~FAndroidEOSSDKManager()
+{	
+	FCoreDelegates::ApplicationHasEnteredForegroundDelegate.RemoveAll(this);
+	FCoreDelegates::ApplicationWillEnterBackgroundDelegate.RemoveAll(this);
+}
+
 EOS_EResult FAndroidEOSSDKManager::EOSInitialize(EOS_InitializeOptions& Options)
 {
 	EOS_Android_InitializeOptions SystemInitializeOptions = { 0 };
diff --git a/Engine/Plugins/Online/EOSShared/Source/EOSShared/Private/Android/AndroidEOSSDKManager.h b/Engine/Plugins/Online/EOSShared/Source/EOSShared/Private/Android/AndroidEOSSDKManager.h
--- a/Engine/Plugins/Online/EOSShared/Source/EOSShared/Private/Android/AndroidEOSSDKManager.h
+++ b/Engine/Plugins/Online/EOSShared/Source/EOSShared/Private/Android/AndroidEOSSDKManager.h
@@ -9,6 +9,9 @@
 class FAndroidEOSSDKManager : public FEOSSDKManager
 {
 public:
+	FAndroidEOSSDKManager();
+	~FAndroidEOSSDKManager();
+
 	virtual EOS_EResult EOSInitialize(EOS_InitializeOptions& Options) override;
 	virtual FString GetCacheDirBase() const override;
 };
diff --git a/Engine/Plugins/Online/EOSShared/Source/EOSShared/Private/IOS/IOSEOSSDKManager.cpp b/Engine/Plugins/Online/EOSShared/Source/EOSShared/Private/IOS/IOSEOSSDKManager.cpp
--- a/Engine/Plugins/Online/EOSShared/Source/EOSShared/Private/IOS/IOSEOSSDKManager.cpp
+++ b/Engine/Plugins/Online/EOSShared/Source/EOSShared/Private/IOS/IOSEOSSDKManager.cpp
@@ -1,8 +1,21 @@
 // Copyright Epic Games, Inc. All Rights Reserved.
+#include "IOSEOSSDKManager.h"
 
 #if WITH_EOS_SDK
 
-#include "IOSEOSSDKManager.h"
+#include "Misc/CoreDelegates.h"
+
+FIOSEOSSDKManager::FIOSEOSSDKManager()
+{
+	FCoreDelegates::ApplicationHasEnteredForegroundDelegate.AddRaw(this, &FIOSEOSSDKManager::OnApplicationStatusChanged, EOS_EApplicationStatus::EOS_AS_Foreground);
+	FCoreDelegates::ApplicationWillEnterBackgroundDelegate.AddRaw(this, &FIOSEOSSDKManager::OnApplicationStatusChanged, EOS_EApplicationStatus::EOS_AS_BackgroundSuspended);
+}
+
+FIOSEOSSDKManager::~FIOSEOSSDKManager()
+{	
+	FCoreDelegates::ApplicationHasEnteredForegroundDelegate.RemoveAll(this);
+	FCoreDelegates::ApplicationWillEnterBackgroundDelegate.RemoveAll(this);
+}
 
 FString FIOSEOSSDKManager::GetCacheDirBase() const
 {
diff --git a/Engine/Plugins/Online/EOSShared/Source/EOSShared/Private/IOS/IOSEOSSDKManager.h b/Engine/Plugins/Online/EOSShared/Source/EOSShared/Private/IOS/IOSEOSSDKManager.h
--- a/Engine/Plugins/Online/EOSShared/Source/EOSShared/Private/IOS/IOSEOSSDKManager.h
+++ b/Engine/Plugins/Online/EOSShared/Source/EOSShared/Private/IOS/IOSEOSSDKManager.h
@@ -8,6 +8,10 @@
 
 class FIOSEOSSDKManager : public FEOSSDKManager
 {
+public:
+	FIOSEOSSDKManager();
+	~FIOSEOSSDKManager();
+
 	virtual FString GetCacheDirBase() const override;
 };
 
