- Fix log messages being emitted during EOSSDKManager shutdown not being suppressed

diff --git a/Engine/Plugins/Online/EOSShared/Source/EOSShared/Private/EOSSDKManager.cpp b/Engine/Plugins/Online/EOSShared/Source/EOSShared/Private/EOSSDKManager.cpp
--- a/Engine/Plugins/Online/EOSShared/Source/EOSShared/Private/EOSSDKManager.cpp
+++ b/Engine/Plugins/Online/EOSShared/Source/EOSShared/Private/EOSSDKManager.cpp
@@ -4,6 +4,7 @@
 
 #if WITH_EOS_SDK
 
+#include "Algo/AnyOf.h"
 #include "Containers/Ticker.h"
 #include "HAL/LowLevelMemTracker.h"
 #include "Misc/App.h"
@@ -19,6 +20,7 @@
 #include "CoreGlobals.h"
 
 #include "EOSShared.h"
+#include "EOSSharedModule.h"
 #include "EOSSharedTypes.h"
 
 #include "eos_auth.h"
@@ -88,16 +90,9 @@ namespace
 		bool bSuppressLogLevel = false;
 		if (IsLogLevelSuppressable(Message->Level))
 		{
-			if (FEOSSDKManager* Manager = static_cast<FEOSSDKManager*>(IEOSSDKManager::Get()))
+			if(FEOSSharedModule* Module = FEOSSharedModule::Get())
 			{
-				for (const FString& SuppressedString : Manager->GetSuppressedLogStrings())
-				{
-					if (MessageStr.Contains(SuppressedString))
-					{
-						bSuppressLogLevel = true;
-						break;
-					}
-				}
+				bSuppressLogLevel = Algo::AnyOf(Module->GetSuppressedLogStrings(), [&MessageStr](const FString& SuppressedLogString) { return MessageStr.Contains(SuppressedLogString); });
 			}
 		}
 
@@ -524,8 +519,6 @@ void FEOSSDKManager::LoadConfig()
 	ConfigTickIntervalSeconds = 0.f;
 	GConfig->GetDouble(SectionName, TEXT("TickIntervalSeconds"), ConfigTickIntervalSeconds, GEngineIni);
 
-	GConfig->GetArray(SectionName, TEXT("SuppressedLogStrings"), SuppressedLogStrings, GEngineIni);
-
 	SetupTicker();
 }
 
diff --git a/Engine/Plugins/Online/EOSShared/Source/EOSShared/Private/EOSSDKManager.h b/Engine/Plugins/Online/EOSShared/Source/EOSShared/Private/EOSSDKManager.h
--- a/Engine/Plugins/Online/EOSShared/Source/EOSShared/Private/EOSSDKManager.h
+++ b/Engine/Plugins/Online/EOSShared/Source/EOSShared/Private/EOSSDKManager.h
@@ -108,11 +108,6 @@ protected:
 	// Config
 	/** Interval between platform ticks. 0 means we tick every frame. */
 	double ConfigTickIntervalSeconds = 0.f;
-public:
-	const TArray<FString>& GetSuppressedLogStrings() const { return SuppressedLogStrings; }
-protected:
-	/** Log strings that should be downgraded from Error/Warning to Log. */
-	TArray<FString> SuppressedLogStrings;
 };
 
 struct FEOSPlatformHandle : public IEOSPlatformHandle
diff --git a/Engine/Plugins/Online/EOSShared/Source/EOSShared/Private/EOSSharedModule.cpp b/Engine/Plugins/Online/EOSShared/Source/EOSShared/Private/EOSSharedModule.cpp
--- a/Engine/Plugins/Online/EOSShared/Source/EOSShared/Private/EOSSharedModule.cpp
+++ b/Engine/Plugins/Online/EOSShared/Source/EOSShared/Private/EOSSharedModule.cpp
@@ -13,11 +13,16 @@
 
 #define LOCTEXT_NAMESPACE "EOS"
 
+#define CONFIG_SECTION_NAME TEXT("EOSSDK")
+
 IMPLEMENT_MODULE(FEOSSharedModule, EOSShared);
 
 void FEOSSharedModule::StartupModule()
 {
 #if WITH_EOS_SDK
+	FCoreDelegates::TSOnConfigSectionsChanged().AddRaw(this, &FEOSSharedModule::OnConfigSectionsChanged);
+	LoadConfig();
+
 	SDKManager = MakeUnique<FPlatformEOSSDKManager>();
 	check(SDKManager);
 
@@ -39,6 +44,8 @@ void FEOSSharedModule::StartupModule()
 void FEOSSharedModule::ShutdownModule()
 {
 #if WITH_EOS_SDK
+	FCoreDelegates::TSOnConfigSectionsChanged().RemoveAll(this);
+
 	if(SDKManager.IsValid())
 	{
 		IModularFeatures::Get().UnregisterModularFeature(TEXT("EOSSDKManager"), SDKManager.Get());
@@ -48,4 +55,23 @@ void FEOSSharedModule::ShutdownModule()
 #endif // WITH_EOS_SDK
 }
 
+FEOSSharedModule* FEOSSharedModule::Get()
+{
+	return FModuleManager::GetModulePtr<FEOSSharedModule>("EOSShared");
+}
+
+void FEOSSharedModule::OnConfigSectionsChanged(const FString& IniFilename, const TSet<FString>& SectionNames)
+{
+	if (IniFilename == GEngineIni && SectionNames.Contains(CONFIG_SECTION_NAME))
+	{
+		LoadConfig();
+	}
+}
+
+void FEOSSharedModule::LoadConfig()
+{
+	GConfig->GetArray(CONFIG_SECTION_NAME, TEXT("SuppressedLogStrings"), SuppressedLogStrings, GEngineIni);
+}
+
+#undef CONFIG_SECTION_NAME
 #undef LOCTEXT_NAMESPACE
\ No newline at end of file
diff --git a/Engine/Plugins/Online/EOSShared/Source/EOSShared/Private/EOSSharedModule.h b/Engine/Plugins/Online/EOSShared/Source/EOSShared/Private/EOSSharedModule.h
--- a/Engine/Plugins/Online/EOSShared/Source/EOSShared/Private/EOSSharedModule.h
+++ b/Engine/Plugins/Online/EOSShared/Source/EOSShared/Private/EOSSharedModule.h
@@ -12,13 +12,22 @@ public:
 	FEOSSharedModule() = default;
 	~FEOSSharedModule() = default;
 
+	static FEOSSharedModule* Get();
+
+	const TArray<FString>& GetSuppressedLogStrings() const { return SuppressedLogStrings; }
+
 private:
 	// ~Begin IModuleInterface
 	virtual void StartupModule() override;
 	virtual void ShutdownModule() override;
 	// ~End IModuleInterface
 
+	void OnConfigSectionsChanged(const FString& IniFilename, const TSet<FString>& SectionNames);
+	void LoadConfig();
+
+
 #if WITH_EOS_SDK
 	TUniquePtr<FEOSSDKManager> SDKManager;
 #endif
+	TArray<FString> SuppressedLogStrings;
 };
\ No newline at end of file
