- Reject attempts to init the SDK prior to forking in fork-requested runs

diff --git a/Engine/Plugins/Online/EOSShared/Source/EOSShared/Private/EOSSDKManager.cpp b/Engine/Plugins/Online/EOSShared/Source/EOSShared/Private/EOSSDKManager.cpp
--- a/Engine/Plugins/Online/EOSShared/Source/EOSShared/Private/EOSSDKManager.cpp
+++ b/Engine/Plugins/Online/EOSShared/Source/EOSShared/Private/EOSSDKManager.cpp
@@ -10,6 +10,7 @@
 #include "Misc/App.h"
 #include "Misc/CoreMisc.h"
 #include "Misc/CoreDelegates.h"
+#include "Misc/Fork.h"
 #include "Misc/NetworkVersion.h"
 #include "Misc/Paths.h"
 #include "Misc/ConfigCacheIni.h"
@@ -185,15 +186,6 @@ static FDelayedAutoRegisterHelper GKickoffDll(EDelayedRegisterRunPhase::EOS_DLL_
 
 FEOSSDKManager::FEOSSDKManager()
 {
-#if EOSSDK_RUNTIME_LOAD_REQUIRED
-	LLM_SCOPE(ELLMTag::RealTimeCommunications);
-	SDKHandle = GetSdkDllHandle();
-	if (SDKHandle == nullptr)
-	{
-		UE_LOG(LogEOSSDK, Warning, TEXT("Unable to load EOSSDK dynamic library"));
-		return;
-	}
-#endif
 }
 
 FEOSSDKManager::~FEOSSDKManager()
@@ -208,7 +200,18 @@ FEOSSDKManager::~FEOSSDKManager()
 
 EOS_EResult FEOSSDKManager::Initialize()
 {
+	if (FForkProcessHelper::IsForkRequested() && !FForkProcessHelper::IsForkedChildProcess())
+	{
+		UE_LOG(LogEOSSDK, Error, TEXT("FEOSSDKManager::Initialize failed, pre-fork"));
+		return EOS_EResult::EOS_InvalidState;
+	}
+
 #if EOSSDK_RUNTIME_LOAD_REQUIRED
+	if (SDKHandle == nullptr)
+	{
+		LLM_SCOPE(ELLMTag::RealTimeCommunications);
+		SDKHandle = GetSdkDllHandle();
+	}
 	if (SDKHandle == nullptr)
 	{
 		UE_LOG(LogEOSSDK, Log, TEXT("FEOSSDKManager::Initialize failed, SDKHandle=nullptr"));
