Make EOSShared responsible for initializing the EOSSDK.

diff --git a/Engine/Plugins/Online/EOSShared/Source/EOSShared/Private/EOSSDKManager.h b/Engine/Plugins/Online/EOSShared/Source/EOSShared/Private/EOSSDKManager.h
--- a/Engine/Plugins/Online/EOSShared/Source/EOSShared/Private/EOSSDKManager.h
+++ b/Engine/Plugins/Online/EOSShared/Source/EOSShared/Private/EOSSDKManager.h
@@ -29,7 +29,9 @@ public:
 	virtual ~FEOSSDKManager();
 
 	// Begin IEOSSDKManager
+	PRAGMA_DISABLE_DEPRECATION_WARNINGS
 	virtual EOS_EResult Initialize() override;
+	PRAGMA_ENABLE_DEPRECATION_WARNINGS
 	virtual bool IsInitialized() const override { return bInitialized; }
 
 	virtual const FEOSSDKPlatformConfig* GetPlatformConfig(const FString& PlatformConfigName, bool bLoadIfMissing = false) override;
diff --git a/Engine/Plugins/Online/EOSShared/Source/EOSShared/Private/EOSSharedModule.cpp b/Engine/Plugins/Online/EOSShared/Source/EOSShared/Private/EOSSharedModule.cpp
--- a/Engine/Plugins/Online/EOSShared/Source/EOSShared/Private/EOSSharedModule.cpp
+++ b/Engine/Plugins/Online/EOSShared/Source/EOSShared/Private/EOSSharedModule.cpp
@@ -26,7 +26,7 @@ void FEOSSharedModule::StartupModule()
 	SDKManager = MakeUnique<FPlatformEOSSDKManager>();
 	check(SDKManager);
 
-	IModularFeatures::Get().RegisterModularFeature(TEXT("EOSSDKManager"), SDKManager.Get());
+	IModularFeatures::Get().RegisterModularFeature(IEOSSDKManager::GetModularFeatureName(), SDKManager.Get());
 
 	// Load from a configurable array of modules at this point, so things that need to bind to the SDK Manager init hooks can do so.
 	TArray<FString> ModulesToLoad;
@@ -38,6 +38,35 @@ void FEOSSharedModule::StartupModule()
 			FModuleManager::Get().LoadModule(*ModuleToLoad);
 		}
 	}
+
+	// When the process starts, initialize the EOSSDK as soon as possible. SDK initialization
+	// modifies environment variables which may cause a crash if another thread is iterating them
+	// at the same time.
+	// The SDK does not support forking - When forking is enabled, the parent process must not
+	// initialize the SDK so that the child process can initialize it properly.
+	if (FForkProcessHelper::IsForkRequested())
+	{
+		// The process has forking enabled. Wait for the fork to occur before initializing the SDK.
+		// Only the child process should initialize the SDK so the parent can continue spawning
+		// processes.
+		OnPostForkDelegateHandle = FCoreDelegates::OnPostFork.AddLambda([](EForkProcessRole ProcessRole)
+		{
+			if (ProcessRole == EForkProcessRole::Child)
+			{
+				if (FEOSSDKManager* SDKManagerPtr = static_cast<FEOSSDKManager*>(FEOSSDKManager::Get()))
+				{
+					SDKManagerPtr->Initialize();
+				}
+			}
+		});
+	}
+	else
+	{
+		if (SDKManager.IsValid())
+		{
+			SDKManager->Initialize();
+		}
+	}
 #endif // WITH_EOS_SDK
 }
 
@@ -45,10 +74,11 @@ void FEOSSharedModule::ShutdownModule()
 {
 #if WITH_EOS_SDK
 	FCoreDelegates::TSOnConfigSectionsChanged().RemoveAll(this);
+	FCoreDelegates::OnPostFork.Remove(OnPostForkDelegateHandle);
 
-	if(SDKManager.IsValid())
+	if (SDKManager.IsValid())
 	{
-		IModularFeatures::Get().UnregisterModularFeature(TEXT("EOSSDKManager"), SDKManager.Get());
+		IModularFeatures::Get().UnregisterModularFeature(IEOSSDKManager::GetModularFeatureName(), SDKManager.Get());
 		SDKManager->Shutdown();
 		SDKManager.Reset();
 	}
diff --git a/Engine/Plugins/Online/EOSShared/Source/EOSShared/Private/EOSSharedModule.h b/Engine/Plugins/Online/EOSShared/Source/EOSShared/Private/EOSSharedModule.h
--- a/Engine/Plugins/Online/EOSShared/Source/EOSShared/Private/EOSSharedModule.h
+++ b/Engine/Plugins/Online/EOSShared/Source/EOSShared/Private/EOSSharedModule.h
@@ -2,6 +2,7 @@
 
 #pragma once
 
+#include "Misc/Fork.h"
 #include "Modules/ModuleInterface.h"
 
 #include "EOSSDKManager.h"
@@ -30,4 +31,5 @@ private:
 	TUniquePtr<FEOSSDKManager> SDKManager;
 #endif
 	TArray<FString> SuppressedLogStrings;
+	FDelegateHandle OnPostForkDelegateHandle;
 };
\ No newline at end of file
diff --git a/Engine/Plugins/Online/EOSShared/Source/EOSShared/Public/IEOSSDKManager.h b/Engine/Plugins/Online/EOSShared/Source/EOSShared/Public/IEOSSDKManager.h
--- a/Engine/Plugins/Online/EOSShared/Source/EOSShared/Public/IEOSSDKManager.h
+++ b/Engine/Plugins/Online/EOSShared/Source/EOSShared/Public/IEOSSDKManager.h
@@ -89,6 +89,7 @@ public:
 
 	virtual ~IEOSSDKManager() = default;
 
+	UE_DEPRECATED(5.3, "Initialize is now called automatically when the EOSShared module is started. User code should instead check IsInitialized to see whether the EOSSDK was successfully started.")
 	virtual EOS_EResult Initialize() = 0;
 	virtual bool IsInitialized() const = 0;
 
diff --git a/Engine/Plugins/Online/OnlineServicesEOS/Source/Private/Online/Tests/OnlineIdEOSTests.cpp b/Engine/Plugins/Online/OnlineServicesEOS/Source/Private/Online/Tests/OnlineIdEOSTests.cpp
--- a/Engine/Plugins/Online/OnlineServicesEOS/Source/Private/Online/Tests/OnlineIdEOSTests.cpp
+++ b/Engine/Plugins/Online/OnlineServicesEOS/Source/Private/Online/Tests/OnlineIdEOSTests.cpp
@@ -11,7 +11,7 @@ IMPLEMENT_SIMPLE_AUTOMATION_TEST(FAccountIdReplicationTest,
 bool FAccountIdReplicationTest::RunTest(const FString& Parameters)
 {
 	IEOSSDKManager* SDKManager = IEOSSDKManager::Get();
-	if(SDKManager && SDKManager->Initialize() == EOS_EResult::EOS_Success)
+	if(ensure(SDKManager) && ensure(SDKManager->IsInitialized()))
 	{
 		using namespace UE::Online;
 
diff --git a/Engine/Plugins/Online/OnlineServicesEOSGS/Source/Private/Online/OnlineServicesEOSGSPlatformFactory.cpp b/Engine/Plugins/Online/OnlineServicesEOSGS/Source/Private/Online/OnlineServicesEOSGSPlatformFactory.cpp
--- a/Engine/Plugins/Online/OnlineServicesEOSGS/Source/Private/Online/OnlineServicesEOSGSPlatformFactory.cpp
+++ b/Engine/Plugins/Online/OnlineServicesEOSGS/Source/Private/Online/OnlineServicesEOSGSPlatformFactory.cpp
@@ -80,12 +80,13 @@ IEOSPlatformHandlePtr FOnlineServicesEOSGSPlatformFactory::CreatePlatform()
 	IEOSSDKManager* const SDKManager = IEOSSDKManager::Get();
 	if (!SDKManager)
 	{
+		UE_LOG(LogOnlineServices, Error, TEXT("[FOnlineServicesEOSGS::Initialize] EOSSDK has not been loaded."));
 		return {};
 	}
 
-	EOS_EResult InitResult = SDKManager->Initialize();
-	if (InitResult != EOS_EResult::EOS_Success)
+	if (!SDKManager->IsInitialized())
 	{
+		UE_LOG(LogOnlineServices, Error, TEXT("[FOnlineServicesEOSGS::Initialize] EOSSDK has not been initialized."));
 		return {};
 	}
 
diff --git a/Engine/Plugins/Online/OnlineSubsystemEOS/Source/OnlineSubsystemEOS/Private/OnlineSubsystemEOS.cpp b/Engine/Plugins/Online/OnlineSubsystemEOS/Source/OnlineSubsystemEOS/Private/OnlineSubsystemEOS.cpp
--- a/Engine/Plugins/Online/OnlineSubsystemEOS/Source/OnlineSubsystemEOS/Private/OnlineSubsystemEOS.cpp
+++ b/Engine/Plugins/Online/OnlineSubsystemEOS/Source/OnlineSubsystemEOS/Private/OnlineSubsystemEOS.cpp
@@ -147,10 +147,9 @@ void FOnlineSubsystemEOS::ModuleInit()
 	// If a fork is requested, we need to wait for post-fork to init the SDK
 	if (!FForkProcessHelper::IsForkRequested() || FForkProcessHelper::IsForkedChildProcess())
 	{
-		EOS_EResult InitResult = EOSSDKManager->Initialize();
-		if (InitResult != EOS_EResult::EOS_Success)
+		if (!EOSSDKManager->IsInitialized())
 		{
-			UE_LOG_ONLINE(Error, TEXT("FOnlineSubsystemEOS: failed to initialize the EOS SDK with result code (%d)"), InitResult);
+			UE_LOG_ONLINE(Error, TEXT("FOnlineSubsystemEOS: EOSSDK not initialized"));
 			return;
 		}
 	}
@@ -281,10 +280,9 @@ bool FOnlineSubsystemEOS::Init()
 		return false;
 	}
 
-	EOS_EResult InitResult = EOSSDKManager->Initialize();
-	if (InitResult != EOS_EResult::EOS_Success)
+	if (!EOSSDKManager->IsInitialized())
 	{
-		UE_LOG_ONLINE(Error, TEXT("FOnlineSubsystemEOS: failed to initialize the EOS SDK with result code (%d)"), InitResult);
+		UE_LOG_ONLINE(Error, TEXT("FOnlineSubsystemEOS: EOSSDK not initialized"));
 		return false;
 	}
 
diff --git a/Engine/Plugins/Online/VoiceChat/EOSVoiceChat/Source/EOSVoiceChat/Private/EOSVoiceChat.cpp b/Engine/Plugins/Online/VoiceChat/EOSVoiceChat/Source/EOSVoiceChat/Private/EOSVoiceChat.cpp
--- a/Engine/Plugins/Online/VoiceChat/EOSVoiceChat/Source/EOSVoiceChat/Private/EOSVoiceChat.cpp
+++ b/Engine/Plugins/Online/VoiceChat/EOSVoiceChat/Source/EOSVoiceChat/Private/EOSVoiceChat.cpp
@@ -102,8 +102,7 @@ void FEOSVoiceChat::Initialize(const FOnVoiceChatInitializeCompleteDelegate& Ini
 
 			if (!EosPlatformHandle)
 			{
-				EOS_EResult EosResult = SDKManager.Initialize();
-				if (EosResult == EOS_EResult::EOS_Success)
+				if (SDKManager.IsInitialized())
 				{
 					FString ConfigProductId;
 					FString ConfigSandboxId;
@@ -167,7 +166,7 @@ void FEOSVoiceChat::Initialize(const FOnVoiceChatInitializeCompleteDelegate& Ini
 				}
 				else
 				{
-					UE_LOG(LogEOSVoiceChat, Warning, TEXT("FEOSVoiceChat::Initialize SDKManager.Initialize failed"));
+					UE_LOG(LogEOSVoiceChat, Warning, TEXT("FEOSVoiceChat::Initialize SDKManager not initialized"));
 					Result = FVoiceChatResult(EVoiceChatResult::ImplementationError);
 				}
 			}
