Use correct value for EOSShared CacheDirectory.

diff --git a/Engine/Plugins/Online/EOSShared/Source/EOSShared/Private/EOSSDKManager.cpp b/Engine/Plugins/Online/EOSShared/Source/EOSShared/Private/EOSSDKManager.cpp
--- a/Engine/Plugins/Online/EOSShared/Source/EOSShared/Private/EOSSDKManager.cpp
+++ b/Engine/Plugins/Online/EOSShared/Source/EOSShared/Private/EOSSDKManager.cpp
@@ -451,7 +451,7 @@ IEOSPlatformHandlePtr FEOSSDKManager::CreatePlatform(const FString& PlatformConf
 
 	if (FPlatformMisc::IsCacheStorageAvailable())
 	{
-		PlatformOptions.CacheDirectory = Utf8CacheDirectory.Length() ? Utf8DeploymentId.Get() : nullptr;
+		PlatformOptions.CacheDirectory = Utf8CacheDirectory.Length() ? Utf8CacheDirectory.Get() : nullptr;
 	}
 	else
 	{
