Add additional runtime options to EOSShared config, update Lyra to use new config.

diff --git a/Engine/Plugins/Online/EOSShared/Source/EOSShared/Private/EOSSDKManager.cpp b/Engine/Plugins/Online/EOSShared/Source/EOSShared/Private/EOSSDKManager.cpp
--- a/Engine/Plugins/Online/EOSShared/Source/EOSShared/Private/EOSSDKManager.cpp
+++ b/Engine/Plugins/Online/EOSShared/Source/EOSShared/Private/EOSSDKManager.cpp
@@ -316,7 +316,35 @@ const FEOSSDKPlatformConfig* FEOSSDKManager::GetPlatformConfig(const FString& Pl
 	GConfig->GetString(*SectionName, TEXT("OverrideCountryCode"), PlatformConfig->OverrideCountryCode, GEngineIni);
 	GConfig->GetString(*SectionName, TEXT("OverrideLocaleCode"), PlatformConfig->OverrideLocaleCode, GEngineIni);
 	GConfig->GetString(*SectionName, TEXT("DeploymentId"), PlatformConfig->DeploymentId, GEngineIni);
+
+	if (GConfig->GetString(*SectionName, TEXT("CacheBaseSubdirectory"), PlatformConfig->CacheDirectory, GEngineIni))
+	{
+		PlatformConfig->CacheDirectory = GetCacheDirBase() / PlatformConfig->CacheDirectory;
+	}
 	GConfig->GetString(*SectionName, TEXT("CacheDirectory"), PlatformConfig->CacheDirectory, GEngineIni);
+
+	bool bCheckRuntimeType = false;
+	GConfig->GetBool(*SectionName, TEXT("bCheckRuntimeType"), bCheckRuntimeType, GEngineIni);
+	if (bCheckRuntimeType)
+	{
+		PlatformConfig->bIsServer = IsRunningDedicatedServer();
+		PlatformConfig->bLoadingInEditor = !IsRunningGame() && !IsRunningDedicatedServer();
+
+		if (PlatformConfig->bIsServer || PlatformConfig->bLoadingInEditor)
+		{
+			// Don't attempt to load overlay for servers or editors.
+			PlatformConfig->bDisableOverlay = true;
+			PlatformConfig->bDisableSocialOverlay = true;
+		}
+		else
+		{
+			// Overlay is on by default, enable additional overlay options.
+			PlatformConfig->bWindowsEnableOverlayD3D9 = true;
+			PlatformConfig->bWindowsEnableOverlayD3D10 = true;
+			PlatformConfig->bWindowsEnableOverlayOpenGL = true;
+		}
+	}
+
 	GConfig->GetBool(*SectionName, TEXT("bIsServer"), PlatformConfig->bIsServer, GEngineIni);
 	GConfig->GetBool(*SectionName, TEXT("bLoadingInEditor"), PlatformConfig->bLoadingInEditor, GEngineIni);
 	GConfig->GetBool(*SectionName, TEXT("bDisableOverlay"), PlatformConfig->bDisableOverlay, GEngineIni);
