Crash in VoiceChat after failed LeaveRoom during shutdown
Speculative fix ensuring target object lifespan during callback.

diff --git a/Engine/Plugins/Online/EOSShared/Source/EOSShared/Public/EOSSharedTypes.h b/Engine/Plugins/Online/EOSShared/Source/EOSShared/Public/EOSSharedTypes.h
--- a/Engine/Plugins/Online/EOSShared/Source/EOSShared/Public/EOSSharedTypes.h
+++ b/Engine/Plugins/Online/EOSShared/Source/EOSShared/Public/EOSSharedTypes.h
@@ -63,7 +63,8 @@ private:
 			check(IsInGameThread());
 		}
 
-		if (CallbackThis->Owner.IsValid())
+		TSharedPtr <OwningType> Pinned = CallbackThis->Owner.Pin();
+		if (Pinned.IsValid())
 		{
 			check(CallbackThis->CallbackLambda);
 
