- Change default TickBudget to 0 (no limit)

diff --git a/Engine/Plugins/Online/EOSShared/Source/EOSShared/Public/IEOSSDKManager.h b/Engine/Plugins/Online/EOSShared/Source/EOSShared/Public/IEOSSDKManager.h
--- a/Engine/Plugins/Online/EOSShared/Source/EOSShared/Public/IEOSSDKManager.h
+++ b/Engine/Plugins/Online/EOSShared/Source/EOSShared/Public/IEOSSDKManager.h
@@ -33,7 +33,7 @@ struct FEOSSDKPlatformConfig
 	bool bWindowsEnableOverlayD3D10 = false;
 	bool bWindowsEnableOverlayOpenGL = false;
 	bool bEnableRTC = true;
-	int32 TickBudgetInMilliseconds = 1;
+	int32 TickBudgetInMilliseconds = 0;
 	TArray<FString> OptionalConfig;
 };
 
