Added UE_CALL_ONCE wrapping for EOSSDKManager log

diff --git a/Engine/Plugins/Online/EOSShared/Source/EOSShared/Private/EOSSDKManager.cpp b/Engine/Plugins/Online/EOSShared/Source/EOSShared/Private/EOSSDKManager.cpp
--- a/Engine/Plugins/Online/EOSShared/Source/EOSShared/Private/EOSSDKManager.cpp
+++ b/Engine/Plugins/Online/EOSShared/Source/EOSShared/Private/EOSSDKManager.cpp
@@ -706,7 +706,7 @@ void FEOSSDKManager::SetupTicker()
 
 void FEOSSDKManager::OnBackBufferReady_RenderThread(SWindow& SlateWindow, const FTexture2DRHIRef& InBackBuffer)
 {
-	UE_LOG(LogEOSSDK, VeryVerbose, TEXT("[%hs] The method is not implemented for this platform."), __FUNCTION__);
+	UE_CALL_ONCE([]() {	UE_LOG(LogEOSSDK, VeryVerbose, TEXT("[%hs] The method is not implemented for this platform."), __FUNCTION__) });
 }
 
 bool FEOSSDKManager::IsRenderReady()
