- Add message box and fatal log message when EOSSDK fails to be loaded. This is config driven and defaults to off.

diff --git a/Engine/Plugins/Online/EOSShared/Source/EOSShared/Private/EOSSDKManager.cpp b/Engine/Plugins/Online/EOSShared/Source/EOSShared/Private/EOSSDKManager.cpp
--- a/Engine/Plugins/Online/EOSShared/Source/EOSShared/Private/EOSSDKManager.cpp
+++ b/Engine/Plugins/Online/EOSShared/Source/EOSShared/Private/EOSSDKManager.cpp
@@ -159,6 +159,21 @@ namespace
 			Result = FPlatformProcess::GetDllHandle(TEXT(EOSSDK_RUNTIME_LIBRARY_NAME));
 		}
 
+		if (!Result)
+		{
+			bool bDllLoadFailureIsFatal = false;
+			GConfig->GetBool(TEXT("EOSSDK"), TEXT("bDllLoadFailureIsFatal"), bDllLoadFailureIsFatal, GEngineIni);
+			if (bDllLoadFailureIsFatal)
+			{
+				FPlatformMisc::MessageBoxExt(
+					EAppMsgType::Ok,
+					*FText::Format(NSLOCTEXT("EOSShared", "DllLoadFail", "Failed to load {0}. Please verify your installation. Exiting..."), FText::FromString(TEXT(EOSSDK_RUNTIME_LIBRARY_NAME))).ToString(),
+					TEXT("Error")
+				);
+				UE_LOG(LogEOSSDK, Fatal, TEXT("Failed to load EOSSDK binary"));
+			}
+		}
+
 		return Result;
 	}
 #endif
