Reverting change back to having a separate cache value for EOSSDK Switch
Added missing flag to Switch integrated platform options
Set unused parameter to nullptr

diff --git a/Engine/Plugins/Online/EOSShared/Source/EOSShared/Private/EOSSDKManager.cpp b/Engine/Plugins/Online/EOSShared/Source/EOSShared/Private/EOSSDKManager.cpp
--- a/Engine/Plugins/Online/EOSShared/Source/EOSShared/Private/EOSSDKManager.cpp
+++ b/Engine/Plugins/Online/EOSShared/Source/EOSShared/Private/EOSSDKManager.cpp
@@ -468,7 +468,7 @@ void FEOSSDKManager::ApplyIntegratedPlatformOptions(EOS_HIntegratedPlatformOptio
 			PlatformOptions.ApiVersion = 1;
 			UE_EOS_CHECK_API_MISMATCH(EOS_INTEGRATEDPLATFORM_OPTIONS_API_LATEST, 1);
 			PlatformOptions.Type = GetIntegratedPlatformType();
-			PlatformOptions.Flags = EOS_EIntegratedPlatformManagementFlags::EOS_IPMF_LibraryManagedByApplication | EOS_EIntegratedPlatformManagementFlags::EOS_IPMF_DisableSDKManagedSessions;
+			PlatformOptions.Flags = EOS_EIntegratedPlatformManagementFlags::EOS_IPMF_LibraryManagedByApplication | EOS_EIntegratedPlatformManagementFlags::EOS_IPMF_ApplicationManagedIdentityLogin | EOS_EIntegratedPlatformManagementFlags::EOS_IPMF_DisableSDKManagedSessions;
 			PlatformOptions.InitOptions = GetIntegratedPlatformOptions();
 
 			EOS_IntegratedPlatformOptionsContainer_AddOptions AddOptions = {};
