Fix for EOS SDK cache directory to account for system account file path.
Added some checks to set cache directory to nullptr in EOS SDK platform create and empty string in config

diff --git a/Engine/Plugins/Online/EOSShared/Source/EOSShared/Private/EOSSDKManager.cpp b/Engine/Plugins/Online/EOSShared/Source/EOSShared/Private/EOSSDKManager.cpp
--- a/Engine/Plugins/Online/EOSShared/Source/EOSShared/Private/EOSSDKManager.cpp
+++ b/Engine/Plugins/Online/EOSShared/Source/EOSShared/Private/EOSSDKManager.cpp
@@ -337,7 +337,8 @@ const FEOSSDKPlatformConfig* FEOSSDKManager::GetPlatformConfig(const FString& Pl
 
 	if (GConfig->GetString(*SectionName, TEXT("CacheBaseSubdirectory"), PlatformConfig->CacheDirectory, GEngineIni))
 	{
-		PlatformConfig->CacheDirectory = GetCacheDirBase() / PlatformConfig->CacheDirectory;
+		const FString CacheDirBase = GetCacheDirBase();
+		PlatformConfig->CacheDirectory = CacheDirBase.IsEmpty() ? FString() : CacheDirBase / PlatformConfig->CacheDirectory;
 	}
 	GConfig->GetString(*SectionName, TEXT("CacheDirectory"), PlatformConfig->CacheDirectory, GEngineIni);
 
@@ -930,7 +931,15 @@ FString FEOSSDKManager::GetProductVersion() const
 
 FString FEOSSDKManager::GetCacheDirBase() const
 {
-	return FPlatformProcess::UserDir();
+	if (FPlatformMisc::IsCacheStorageAvailable())
+	{
+		return FPlatformProcess::UserDir();
+	}
+	else
+	{
+		return FString(); 
+	}
+	
 }
 
 FString FEOSSDKManager::GetOverrideCountryCode(const EOS_HPlatform Platform) const
diff --git a/Engine/Plugins/Online/EOSShared/Source/EOSShared/Private/Windows/WindowsEOSSDKManager.cpp b/Engine/Plugins/Online/EOSShared/Source/EOSShared/Private/Windows/WindowsEOSSDKManager.cpp
--- a/Engine/Plugins/Online/EOSShared/Source/EOSShared/Private/Windows/WindowsEOSSDKManager.cpp
+++ b/Engine/Plugins/Online/EOSShared/Source/EOSShared/Private/Windows/WindowsEOSSDKManager.cpp
@@ -32,4 +32,18 @@ IEOSPlatformHandlePtr FWindowsEOSSDKManager::CreatePlatform(const FEOSSDKPlatfor
 	return FEOSSDKManager::CreatePlatform(PlatformConfig, PlatformOptions);
 }
 
+FString FWindowsEOSSDKManager::GetCacheDirBase() const
+{
+	if (FPlatformMisc::IsCacheStorageAvailable())
+	{
+		// return folder path in AppData and not Documents to accomodate both user and system accounts
+		return FWindowsPlatformProcess::UserSettingsDir();
+	}
+	else
+	{
+		return FString();
+	}
+
+}
+
 #endif // WITH_EOS_SDK
\ No newline at end of file
diff --git a/Engine/Plugins/Online/EOSShared/Source/EOSShared/Private/Windows/WindowsEOSSDKManager.h b/Engine/Plugins/Online/EOSShared/Source/EOSShared/Private/Windows/WindowsEOSSDKManager.h
--- a/Engine/Plugins/Online/EOSShared/Source/EOSShared/Private/Windows/WindowsEOSSDKManager.h
+++ b/Engine/Plugins/Online/EOSShared/Source/EOSShared/Private/Windows/WindowsEOSSDKManager.h
@@ -10,6 +10,8 @@ class FWindowsEOSSDKManager : public FEOSSDKManager
 {
 protected:
 	virtual IEOSPlatformHandlePtr CreatePlatform(const FEOSSDKPlatformConfig& PlatformConfig, EOS_Platform_Options& PlatformOptions) override;
+public: 
+	virtual FString GetCacheDirBase() const override;
 };
 
 using FPlatformEOSSDKManager = FWindowsEOSSDKManager;
diff --git a/Engine/Plugins/Online/OnlineServicesEOSGS/Source/Private/Online/OnlineServicesEOSGSPlatformFactory.cpp b/Engine/Plugins/Online/OnlineServicesEOSGS/Source/Private/Online/OnlineServicesEOSGSPlatformFactory.cpp
--- a/Engine/Plugins/Online/OnlineServicesEOSGS/Source/Private/Online/OnlineServicesEOSGSPlatformFactory.cpp
+++ b/Engine/Plugins/Online/OnlineServicesEOSGS/Source/Private/Online/OnlineServicesEOSGSPlatformFactory.cpp
@@ -60,7 +60,8 @@ FString LoadEOSPlatformConfig(IEOSSDKManager* const SDKManager)
 	// Config key renamed to ClientEncryptionKey as EncryptionKey gets removed from packaged builds due to IniKeyDenylist=EncryptionKey entry in BaseGame.ini.
 	GConfig->GetString(*ConfigSectionName, TEXT("ClientEncryptionKey"), PlatformConfig.EncryptionKey, GEngineIni);
 
-	PlatformConfig.CacheDirectory = SDKManager->GetCacheDirBase() / TEXT("OnlineServicesEOS");
+	const FString CacheDirBase = SDKManager->GetCacheDirBase();
+	PlatformConfig.CacheDirectory = CacheDirBase.IsEmpty() ? FString() : CacheDirBase / TEXT("OnlineServicesEOS");
 
 	PlatformConfig.bIsServer = IsRunningDedicatedServer() ? EOS_TRUE : EOS_FALSE;
 	if (!IsRunningGame())
diff --git a/Engine/Plugins/Online/OnlineSubsystemEOS/Source/OnlineSubsystemEOS/Private/OnlineSubsystemEOS.cpp b/Engine/Plugins/Online/OnlineSubsystemEOS/Source/OnlineSubsystemEOS/Private/OnlineSubsystemEOS.cpp
--- a/Engine/Plugins/Online/OnlineSubsystemEOS/Source/OnlineSubsystemEOS/Private/OnlineSubsystemEOS.cpp
+++ b/Engine/Plugins/Online/OnlineSubsystemEOS/Source/OnlineSubsystemEOS/Private/OnlineSubsystemEOS.cpp
@@ -226,9 +226,11 @@ bool FOnlineSubsystemEOS::PlatformCreate()
 	// Make the cache directory be in the user's writable area
 	FString CacheDir;
 	
-	if (FPlatformMisc::IsCacheStorageAvailable())
+	const FString CacheDirBase = EOSSDKManager->GetCacheDirBase();
+
+	if (!CacheDirBase.IsEmpty())
 	{
-		CacheDir = EOSSDKManager->GetCacheDirBase() / ArtifactSettings.ArtifactName / EOSSettings.CacheDir;
+		CacheDir = CacheDirBase / ArtifactSettings.ArtifactName / EOSSettings.CacheDir;
 	}
 	const auto CacheDirUtf8 = StringCast<UTF8CHAR>(*CacheDir);
 	PlatformOptions.CacheDirectory = CacheDir.IsEmpty() ? nullptr : (const char*)CacheDirUtf8.Get();
