Added check to avoid log spam when EOS integrated platform is not enabled

diff --git a/Engine/Plugins/Online/EOSShared/Source/EOSShared/Private/EOSSDKManager.cpp b/Engine/Plugins/Online/EOSShared/Source/EOSShared/Private/EOSSDKManager.cpp
--- a/Engine/Plugins/Online/EOSShared/Source/EOSShared/Private/EOSSDKManager.cpp
+++ b/Engine/Plugins/Online/EOSShared/Source/EOSShared/Private/EOSSDKManager.cpp
@@ -779,25 +779,32 @@ void FEOSSDKManager::CallUIPrePresent(const EOS_UI_PrePresentOptions& Options)
 
 bool FEOSSDKManager::IsRenderReady()
 {
-	if (bRenderReady)
+	if (bEnablePlatformIntegration)
 	{
-		return true;
-	}
+		if (bRenderReady)
+		{
+			return true;
+		}
 
-	if (!FSlateApplication::IsInitialized())
-	{
-		return false;
-	}
+		if (!FSlateApplication::IsInitialized())
+		{
+			return false;
+		}
 
-	FSlateRenderer* Renderer = FSlateApplication::Get().GetRenderer();
-	if (!Renderer)
+		FSlateRenderer* Renderer = FSlateApplication::Get().GetRenderer();
+		if (!Renderer)
+		{
+			return false;
+		}
+
+		Renderer->OnBackBufferReadyToPresent().AddRaw(this, &FEOSSDKManager::OnBackBufferReady_RenderThread);
+		bRenderReady = true;
+		return true;
+	}
+	else
 	{
 		return false;
 	}
-
-	Renderer->OnBackBufferReadyToPresent().AddRaw(this, &FEOSSDKManager::OnBackBufferReady_RenderThread);
-	bRenderReady = true;
-	return true;
 }
 
 void FEOSSDKManager::SetInvokeOverlayButton(const EOS_HPlatform PlatformHandle)
