- Log RoomOptions in JoinRoom callback
- Remove an include from EOSShared.h, swap to forward declarations

diff --git a/Engine/Plugins/Online/EOSShared/Source/EOSShared/Private/EOSShared.cpp b/Engine/Plugins/Online/EOSShared/Source/EOSShared/Private/EOSShared.cpp
--- a/Engine/Plugins/Online/EOSShared/Source/EOSShared/Private/EOSShared.cpp
+++ b/Engine/Plugins/Online/EOSShared/Source/EOSShared/Private/EOSShared.cpp
@@ -10,6 +10,8 @@
 #include "eos_auth_types.h"
 #include "eos_friends_types.h"
 #include "eos_presence_types.h"
+#include "eos_rtc_types.h"
+#include "eos_userinfo_types.h"
 
 DEFINE_LOG_CATEGORY(LogEOSSDK);
 
@@ -436,4 +438,11 @@ bool LexFromString(EOS_EIntegratedPlatformManagementFlags& OutEnum, const TCHAR*
 	return true;
 }
 
+FString LexToString(const EOS_RTC_Option& Option)
+{
+	UE_EOS_CHECK_API_MISMATCH(EOS_RTC_OPTION_API_LATEST, 1);
+	check(Option.ApiVersion == 1);
+	return FString::Printf(TEXT("\"%hs\"=\"%hs\""), Option.Key, Option.Value);
+}
+
 #endif // WITH_EOS_SDK
\ No newline at end of file
diff --git a/Engine/Plugins/Online/EOSShared/Source/EOSShared/Public/EOSShared.h b/Engine/Plugins/Online/EOSShared/Source/EOSShared/Public/EOSShared.h
--- a/Engine/Plugins/Online/EOSShared/Source/EOSShared/Public/EOSShared.h
+++ b/Engine/Plugins/Online/EOSShared/Source/EOSShared/Public/EOSShared.h
@@ -38,7 +38,6 @@
 #include "eos_base.h"
 #include "eos_common.h"
 #include "eos_types.h"
-#include "eos_userinfo_types.h"
 #include "eos_version.h"
 
 #ifndef WITH_EOS_RTC
@@ -57,6 +56,11 @@ EOS_ENUM_FORWARD_DECL(EOS_Presence_EStatus);
 EOS_ENUM_FORWARD_DECL(EOS_UI_EInputStateButtonFlags);
 #undef EOS_ENUM_FORWARD_DECL
 
+#define EOS_STRUCT_FORWARD_DECL(name) EXTERN_C typedef struct _tag ## name name;
+EOS_STRUCT_FORWARD_DECL(EOS_UserInfo_BestDisplayName);
+EOS_STRUCT_FORWARD_DECL(EOS_RTC_Option);
+#undef EOS_STRUCT_FORWARD_DECL
+
 DECLARE_LOG_CATEGORY_EXTERN(LogEOSSDK, Log, All);
 
 EOSSHARED_API FString LexToString(const EOS_EResult EosResult);
@@ -91,4 +95,6 @@ EOSSHARED_API bool LexFromString(EOS_EIntegratedPlatformManagementFlags& OutEnum
 /** Extracts the display name FString from a EOS_UserInfo_BestDisplayName using the following logic: Nickname > DisplayNameSanitized > DisplayName */
 EOSSHARED_API FString GetBestDisplayNameStr(const EOS_UserInfo_BestDisplayName& BestDisplayName);
 
+EOSSHARED_API FString LexToString(const EOS_RTC_Option& Option);
+
 #endif // WITH_EOS_SDK
\ No newline at end of file
diff --git a/Engine/Plugins/Online/VoiceChat/EOSVoiceChat/Source/EOSVoiceChat/Private/EOSVoiceChatUser.cpp b/Engine/Plugins/Online/VoiceChat/EOSVoiceChat/Source/EOSVoiceChat/Private/EOSVoiceChatUser.cpp
--- a/Engine/Plugins/Online/VoiceChat/EOSVoiceChat/Source/EOSVoiceChat/Private/EOSVoiceChatUser.cpp
+++ b/Engine/Plugins/Online/VoiceChat/EOSVoiceChat/Source/EOSVoiceChat/Private/EOSVoiceChatUser.cpp
@@ -1686,7 +1686,11 @@ void FEOSVoiceChatUser::OnJoinRoom(const EOS_RTC_JoinRoomCallbackInfo* CallbackI
 		{
 			if (Result.IsSuccess())
 			{
-				EOSVOICECHATUSER_LOG(Log, TEXT("OnJoinRoom ChannelName=[%s] Result=[%s]"), *ChannelName, *LexToString(Result));
+				EOSVOICECHATUSER_LOG(Log, TEXT("OnJoinRoom ChannelName=[%s] Result=[%s] RoomOptions=[%s]"),
+					*ChannelName,
+					*LexToString(Result),
+					*FString::JoinBy(TArrayView<const EOS_RTC_Option>(CallbackInfo->RoomOptions, CallbackInfo->RoomOptionsCount), TEXT(", "), UE_PROJECTION(LexToString))
+				);
 
 				ChannelSession->JoinState = EChannelJoinState::Joined;
 				TriggerCompletionDelegate(ChannelSession->JoinDelegate, ChannelSession->ChannelName, FVoiceChatResult::CreateSuccess());
