- Fix compile errors due to EXTERN_C not always being defined

diff --git a/Engine/Plugins/Online/EOSShared/Source/EOSShared/Public/EOSShared.h b/Engine/Plugins/Online/EOSShared/Source/EOSShared/Public/EOSShared.h
--- a/Engine/Plugins/Online/EOSShared/Source/EOSShared/Public/EOSShared.h
+++ b/Engine/Plugins/Online/EOSShared/Source/EOSShared/Public/EOSShared.h
@@ -56,7 +56,7 @@ EOS_ENUM_FORWARD_DECL(EOS_Presence_EStatus);
 EOS_ENUM_FORWARD_DECL(EOS_UI_EInputStateButtonFlags);
 #undef EOS_ENUM_FORWARD_DECL
 
-#define EOS_STRUCT_FORWARD_DECL(name) EXTERN_C typedef struct _tag ## name name;
+#define EOS_STRUCT_FORWARD_DECL(name) extern "C" typedef struct _tag ## name name;
 EOS_STRUCT_FORWARD_DECL(EOS_UserInfo_BestDisplayName);
 EOS_STRUCT_FORWARD_DECL(EOS_RTC_Option);
 #undef EOS_STRUCT_FORWARD_DECL
