Add conditional Engine depdency to EOSShared
 - Fixes issue with projects requiring EOSShared but doesn't depend on Engine

diff --git a/Engine/Plugins/Online/EOSShared/Source/EOSShared/EOSShared.Build.cs b/Engine/Plugins/Online/EOSShared/Source/EOSShared/EOSShared.Build.cs
--- a/Engine/Plugins/Online/EOSShared/Source/EOSShared/EOSShared.Build.cs
+++ b/Engine/Plugins/Online/EOSShared/Source/EOSShared/EOSShared.Build.cs
@@ -29,9 +29,13 @@ public class EOSShared : ModuleRules
 			new string[]
 			{
 				"Core",
-				"EOSSDK",
-				"Slate"
+				"EOSSDK"
 			}
 		);
+
+		if (Target.bCompileAgainstEngine)
+		{
+			PrivateDependencyModuleNames.Add("Slate");
+		}
 	}
 }
diff --git a/Engine/Plugins/Online/EOSShared/Source/EOSShared/Private/EOSSDKManager.cpp b/Engine/Plugins/Online/EOSShared/Source/EOSShared/Private/EOSSDKManager.cpp
--- a/Engine/Plugins/Online/EOSShared/Source/EOSShared/Private/EOSSDKManager.cpp
+++ b/Engine/Plugins/Online/EOSShared/Source/EOSShared/Private/EOSSDKManager.cpp
@@ -6,10 +6,15 @@
 
 #include "Algo/AnyOf.h"
 #include "Containers/Ticker.h"
-#include "Framework/Application/SlateApplication.h"
 #include "HAL/LowLevelMemTracker.h"
+
+#if WITH_ENGINE
+#include "Framework/Application/SlateApplication.h"
 #include "HAL/PlatformInput.h"
 #include "InputCoreTypes.h"
+#endif
+
+#include "Misc/ScopeRWLock.h"
 #include "Misc/App.h"
 #include "Misc/CoreMisc.h"
 #include "Misc/CoreDelegates.h"
@@ -745,10 +750,12 @@ void FEOSSDKManager::SetupTicker()
 	}
 }
 
+#if WITH_ENGINE
 void FEOSSDKManager::OnBackBufferReady_RenderThread(SWindow& SlateWindow, const FTexture2DRHIRef& InBackBuffer)
 {
 	UE_CALL_ONCE([]() {	UE_LOG(LogEOSSDK, VeryVerbose, TEXT("[%hs] The method is not implemented for this platform."), __FUNCTION__) });
 }
+#endif
 
 void FEOSSDKManager::CallUIPrePresent(const EOS_UI_PrePresentOptions& Options)
 {
@@ -777,6 +784,7 @@ void FEOSSDKManager::CallUIPrePresent(const EOS_UI_PrePresentOptions& Options)
 	}
 }
 
+#if WITH_ENGINE
 bool FEOSSDKManager::IsRenderReady()
 {
 	if (bEnablePlatformIntegration)
@@ -806,6 +814,7 @@ bool FEOSSDKManager::IsRenderReady()
 		return false;
 	}
 }
+#endif
 
 void FEOSSDKManager::SetInvokeOverlayButton(const EOS_HPlatform PlatformHandle)
 {
@@ -831,7 +840,9 @@ bool FEOSSDKManager::Tick(float)
 {
 	check(IsInGameThread());
 
+#if WITH_ENGINE
 	IsRenderReady();
+#endif
 
 	ReleaseReleasedPlatforms();
 
@@ -1051,6 +1062,7 @@ void FEOSSDKManager::Shutdown()
 
 		FCoreDelegates::OnNetworkConnectionStatusChanged.RemoveAll(this);
 
+#if WITH_ENGINE
 		// We can't check bRenderReady at this point as Slate might have shut down already
 		if (FSlateApplication::IsInitialized())
 		{
@@ -1060,6 +1072,7 @@ void FEOSSDKManager::Shutdown()
 				Renderer->OnBackBufferReadyToPresent().RemoveAll(this);
 			}
 		}
+#endif
 	}
 }
 
diff --git a/Engine/Plugins/Online/EOSShared/Source/EOSShared/Private/EOSSDKManager.h b/Engine/Plugins/Online/EOSShared/Source/EOSShared/Private/EOSSDKManager.h
--- a/Engine/Plugins/Online/EOSShared/Source/EOSShared/Private/EOSSDKManager.h
+++ b/Engine/Plugins/Online/EOSShared/Source/EOSShared/Private/EOSSDKManager.h
@@ -8,8 +8,11 @@
 #include "Containers/UnrealString.h"
 #include "IEOSSDKManager.h"
 #include "Misc/CoreMisc.h"
+
+#if WITH_ENGINE
 #include "Widgets/SWindow.h"
 #include "Rendering/SlateRenderer.h"
+#endif
 
 #if defined(EOS_PLATFORM_BASE_FILE_NAME)
 #include EOS_PLATFORM_BASE_FILE_NAME
@@ -77,6 +80,8 @@ protected:
 	virtual bool Tick(float);
 	virtual const void* GetIntegratedPlatformOptions();
 	virtual EOS_IntegratedPlatformType GetIntegratedPlatformType();
+
+#if WITH_ENGINE
 	/** Provided to `OnBackBufferReadyToPresent` to get access to the render thread. */
 	virtual void OnBackBufferReady_RenderThread(SWindow& SlateWindow, const FTexture2DRHIRef& BackBuffer);
 	/**
@@ -84,6 +89,8 @@ protected:
 	 * This will also add the Back Buffer Ready To Present handler.
 	 */
 	virtual bool IsRenderReady();
+#endif
+
 	void SetInvokeOverlayButton(const EOS_HPlatform PlatformHandle);
 	EOS_HIntegratedPlatformOptionsContainer CreateIntegratedPlatformOptionsContainer();
 	void ApplyIntegratedPlatformOptions(EOS_HIntegratedPlatformOptionsContainer& Container);
