Remove reference to change not for release in 5.4

diff --git a/Engine/Plugins/Online/EOSShared/Source/EOSShared/Private/Windows/WindowsEOSSDKManager.h b/Engine/Plugins/Online/EOSShared/Source/EOSShared/Private/Windows/WindowsEOSSDKManager.h
--- a/Engine/Plugins/Online/EOSShared/Source/EOSShared/Private/Windows/WindowsEOSSDKManager.h
+++ b/Engine/Plugins/Online/EOSShared/Source/EOSShared/Private/Windows/WindowsEOSSDKManager.h
@@ -21,9 +21,6 @@ protected:
 
 
 	EOS_IntegratedPlatform_Steam_Options PlatformSteamOptions;
-
-public: 
-	virtual FString GetCacheDirBase() const override;
 };
 
 using FPlatformEOSSDKManager = FWindowsEOSSDKManager;
