- Add LexToUtf8String for EOS_EResult as that's predominantly used for logging, to avoid unnecessary allocations/conversions. Not changing existing callsites for now.

diff --git a/Engine/Plugins/Online/EOSShared/Source/EOSShared/Private/EOSShared.cpp b/Engine/Plugins/Online/EOSShared/Source/EOSShared/Private/EOSShared.cpp
--- a/Engine/Plugins/Online/EOSShared/Source/EOSShared/Private/EOSShared.cpp
+++ b/Engine/Plugins/Online/EOSShared/Source/EOSShared/Private/EOSShared.cpp
@@ -20,6 +20,11 @@ FString LexToString(const EOS_EResult EosResult)
 	return UTF8_TO_TCHAR(EOS_EResult_ToString(EosResult));
 }
 
+const char* LexToUtf8String(const EOS_EResult EosResult)
+{
+	return EOS_EResult_ToString(EosResult);
+}
+
 FString LexToString(const EOS_ProductUserId UserId)
 {
 	FString Result;
diff --git a/Engine/Plugins/Online/EOSShared/Source/EOSShared/Public/EOSShared.h b/Engine/Plugins/Online/EOSShared/Source/EOSShared/Public/EOSShared.h
--- a/Engine/Plugins/Online/EOSShared/Source/EOSShared/Public/EOSShared.h
+++ b/Engine/Plugins/Online/EOSShared/Source/EOSShared/Public/EOSShared.h
@@ -64,6 +64,8 @@ EOS_STRUCT_FORWARD_DECL(EOS_RTC_Option);
 DECLARE_LOG_CATEGORY_EXTERN(LogEOSSDK, Log, All);
 
 EOSSHARED_API FString LexToString(const EOS_EResult EosResult);
+EOSSHARED_API const char* LexToUtf8String(const EOS_EResult EosResult);
+
 EOSSHARED_API FString LexToString(const EOS_ProductUserId UserId);
 EOSSHARED_API void LexFromString(EOS_ProductUserId& UserId, const TCHAR* String);
 inline EOS_ProductUserId EOSProductUserIdFromString(const TCHAR* String)
