- Remove unused EOS_EExternalCredentialType LexFromString

diff --git a/Engine/Plugins/Online/EOSShared/Source/EOSShared/Private/EOSShared.cpp b/Engine/Plugins/Online/EOSShared/Source/EOSShared/Private/EOSShared.cpp
--- a/Engine/Plugins/Online/EOSShared/Source/EOSShared/Private/EOSShared.cpp
+++ b/Engine/Plugins/Online/EOSShared/Source/EOSShared/Private/EOSShared.cpp
@@ -242,45 +242,6 @@ bool LexFromString(EOS_EExternalAccountType& OutEnum, const TCHAR* InString)
 	return true;
 }
 
-bool LexFromString(EOS_EExternalCredentialType& OutEnum, const TCHAR* InString)
-{
-	if (FCString::Stricmp(InString, TEXT("Steam")) == 0)
-	{
-		OutEnum = EOS_EExternalCredentialType::EOS_ECT_STEAM_APP_TICKET;
-	}
-	else if (FCString::Stricmp(InString, TEXT("PSN")) == 0)
-	{
-		OutEnum = EOS_EExternalCredentialType::EOS_ECT_PSN_ID_TOKEN;
-	}
-	else if (FCString::Stricmp(InString, TEXT("Xbox")) == 0)
-	{
-		OutEnum = EOS_EExternalCredentialType::EOS_ECT_XBL_XSTS_TOKEN;
-	}
-	else if (FCString::Stricmp(InString, TEXT("Nintendo")) == 0)
-	{
-		OutEnum = EOS_EExternalCredentialType::EOS_ECT_NINTENDO_ID_TOKEN;
-	}
-	else if (FCString::Stricmp(InString, TEXT("NSA")) == 0)
-	{
-		OutEnum = EOS_EExternalCredentialType::EOS_ECT_NINTENDO_NSA_ID_TOKEN;
-	}
-	else if (FCString::Stricmp(InString, TEXT("Apple")) == 0)
-	{
-		OutEnum = EOS_EExternalCredentialType::EOS_ECT_APPLE_ID_TOKEN;
-	}
-	else if (FCString::Stricmp(InString, TEXT("Google")) == 0)
-	{
-		OutEnum = EOS_EExternalCredentialType::EOS_ECT_GOOGLE_ID_TOKEN;
-	}
-	else
-	{
-		// Unknown means OpenID
-		OutEnum = EOS_EExternalCredentialType::EOS_ECT_OPENID_ACCESS_TOKEN;
-	}
-
-	return true;
-}
-
 bool LexFromString(EOS_EAuthScopeFlags& OutEnum, const FStringView InString)
 {
 	OutEnum = EOS_EAuthScopeFlags::EOS_AS_NoFlags;
diff --git a/Engine/Plugins/Online/EOSShared/Source/EOSShared/Public/EOSShared.h b/Engine/Plugins/Online/EOSShared/Source/EOSShared/Public/EOSShared.h
--- a/Engine/Plugins/Online/EOSShared/Source/EOSShared/Public/EOSShared.h
+++ b/Engine/Plugins/Online/EOSShared/Source/EOSShared/Public/EOSShared.h
@@ -89,7 +89,6 @@ EOSSHARED_API const TCHAR* LexToString(const EOS_Presence_EStatus PresenceStatus
 
 EOSSHARED_API bool LexFromString(EOS_EAuthScopeFlags& OutEnum, const FStringView InString);
 EOSSHARED_API bool LexFromString(EOS_EExternalAccountType& OutEnum, const TCHAR* InString);
-EOSSHARED_API bool LexFromString(EOS_EExternalCredentialType& OutEnum, const TCHAR* InString);
 EOSSHARED_API bool LexFromString(EOS_ELoginCredentialType& OutEnum, const TCHAR* InString);
 EOSSHARED_API bool LexFromString(EOS_ERTCBackgroundMode& OutEnum, const TCHAR* InString);
 EOSSHARED_API bool LexFromString(EOS_UI_EInputStateButtonFlags& OutEnum, const TCHAR* InString);
