- Re-add missing LexFromString(EOS_EExternalAccountType, const TCHAR* InString) removed through bad merge

diff --git a/Engine/Plugins/Online/EOSShared/Source/EOSShared/Public/EOSShared.h b/Engine/Plugins/Online/EOSShared/Source/EOSShared/Public/EOSShared.h
--- a/Engine/Plugins/Online/EOSShared/Source/EOSShared/Public/EOSShared.h
+++ b/Engine/Plugins/Online/EOSShared/Source/EOSShared/Public/EOSShared.h
@@ -91,6 +91,7 @@ EOSSHARED_API const TCHAR* LexToString(const EOS_ENetworkStatus NetworkStatus);
 EOSSHARED_API const TCHAR* LexToString(const EOS_Presence_EStatus PresenceStatus);
 
 EOSSHARED_API bool LexFromString(EOS_EAuthScopeFlags& OutEnum, const FStringView& InString);
+EOSSHARED_API bool LexFromString(EOS_EExternalAccountType& OutEnum, const TCHAR* InString);
 EOSSHARED_API bool LexFromString(EOS_EIntegratedPlatformManagementFlags& OutEnum, const TCHAR* InString);
 EOSSHARED_API bool LexFromString(EOS_ELoginCredentialType& OutEnum, const TCHAR* InString);
 EOSSHARED_API bool LexFromString(EOS_EPacketReliability& OutEnum, const TCHAR* InString);
