- Set EOSSDK Virtual Memory reservation / commit callback methods on the platforms that require them.
- Keep the generic FPlatformMemory implementation of those callbacks in a common place to avoid duplication
- Note the APIs are defined in each eos_<Platform>.h header separately (they're passed in the platform specific Initialization struct) but our implementation is the same on all but one platform so we keep it in a common place.
- Update ApiVersion in various places in line with EOSSDK 1.16.2 update

#tests local builds, preflight, tested on console

diff --git a/Engine/Plugins/Online/EOSShared/Source/EOSShared/Private/EOSSDKManager.cpp b/Engine/Plugins/Online/EOSShared/Source/EOSShared/Private/EOSSDKManager.cpp
--- a/Engine/Plugins/Online/EOSShared/Source/EOSShared/Private/EOSSDKManager.cpp
+++ b/Engine/Plugins/Online/EOSShared/Source/EOSShared/Private/EOSSDKManager.cpp
@@ -557,7 +557,7 @@ IEOSPlatformHandlePtr FEOSSDKManager::CreatePlatform(const FString& PlatformConf
 
 	EOS_Platform_Options PlatformOptions = {};
 	PlatformOptions.ApiVersion = 13;
-	UE_EOS_CHECK_API_MISMATCH(EOS_PLATFORM_OPTIONS_API_LATEST, 13);
+	UE_EOS_CHECK_API_MISMATCH(EOS_PLATFORM_OPTIONS_API_LATEST, 14);
 	PlatformOptions.Reserved = nullptr;
 	ApplySystemSpecificOptions(PlatformOptions.SystemSpecificOptions);
 	ApplyIntegratedPlatformOptions(PlatformOptions.IntegratedPlatformOptionsContainerHandle);
@@ -586,6 +586,7 @@ IEOSPlatformHandlePtr FEOSSDKManager::CreatePlatform(const FString& PlatformConf
 	}
 
 	PlatformOptions.TickBudgetInMilliseconds = PlatformConfig->TickBudgetInMilliseconds;
+	PlatformOptions.TaskNetworkTimeoutSeconds = nullptr;
 
 	EOS_Platform_RTCOptions PlatformRTCOptions = {};
 	PlatformRTCOptions.ApiVersion = 2;
diff --git a/Engine/Plugins/Online/EOSShared/Source/EOSShared/Private/EOSVirtualMemoryCallbacks.cpp b/Engine/Plugins/Online/EOSShared/Source/EOSShared/Private/EOSVirtualMemoryCallbacks.cpp
new file mode 100644
--- /dev/null
+++ b/Engine/Plugins/Online/EOSShared/Source/EOSShared/Private/EOSVirtualMemoryCallbacks.cpp
@@ -0,0 +1,39 @@
+// Copyright Epic Games, Inc. All Rights Reserved.
+
+#include "EOSVirtualMemoryCallbacks.h"
+
+#include "HAL/PlatformMemory.h"
+
+namespace EOSVirtualMemoryCallbacks
+{
+
+void* EOS_MEMORY_CALL Reserve(size_t SizeInBytes, EOS_EVM_AccessType AccessType, void** OutContextData)
+{
+	FPlatformMemory::FPlatformVirtualMemoryBlock* Block = new FPlatformMemory::FPlatformVirtualMemoryBlock(FPlatformMemory::FPlatformVirtualMemoryBlock::AllocateVirtual(SizeInBytes));
+	*OutContextData = Block;
+	return Block->GetVirtualPointer();
+}
+
+EOS_Bool EOS_MEMORY_CALL Release(void* Address, size_t SizeInBytes, void* ContextData)
+{
+	FPlatformMemory::FPlatformVirtualMemoryBlock* Block = static_cast<FPlatformMemory::FPlatformVirtualMemoryBlock*>(ContextData);
+	Block->FreeVirtual();
+	delete Block;
+	return true;
+}
+
+EOS_Bool EOS_MEMORY_CALL Commit(void* Address, size_t SizeInBytes, void* ContextData)
+{
+	FPlatformMemory::FPlatformVirtualMemoryBlock* Block = static_cast<FPlatformMemory::FPlatformVirtualMemoryBlock*>(ContextData);
+	Block->CommitByPtr(Address, SizeInBytes);
+	return true;
+}
+
+EOS_Bool EOS_MEMORY_CALL Decommit(void* Address, size_t SizeInBytes, void* ContextData)
+{
+	FPlatformMemory::FPlatformVirtualMemoryBlock* Block = static_cast<FPlatformMemory::FPlatformVirtualMemoryBlock*>(ContextData);
+	Block->DecommitByPtr(Address, SizeInBytes);
+	return true;
+}
+
+} // namespace EOSVirtualMemoryCallbacks
diff --git a/Engine/Plugins/Online/EOSShared/Source/EOSShared/Private/EOSVirtualMemoryCallbacks.h b/Engine/Plugins/Online/EOSShared/Source/EOSShared/Private/EOSVirtualMemoryCallbacks.h
new file mode 100644
--- /dev/null
+++ b/Engine/Plugins/Online/EOSShared/Source/EOSShared/Private/EOSVirtualMemoryCallbacks.h
@@ -0,0 +1,19 @@
+// Copyright Epic Games, Inc. All Rights Reserved.
+
+#pragma once
+
+#if defined(EOS_PLATFORM_BASE_FILE_NAME)
+#include EOS_PLATFORM_BASE_FILE_NAME
+#endif
+#include "eos_base.h"
+
+enum class EOS_EVM_AccessType : int32_t;
+
+// Generic VMem callback implementations used by most platforms requiring VMem callbacks
+namespace EOSVirtualMemoryCallbacks
+{
+	void* EOS_MEMORY_CALL Reserve(size_t SizeInBytes, EOS_EVM_AccessType AccessType, void** OutContextData);
+	EOS_Bool EOS_MEMORY_CALL Release(void* Address, size_t SizeInBytes, void* ContextData);
+	EOS_Bool EOS_MEMORY_CALL Commit(void* Address, size_t SizeInBytes, void* ContextData);
+	EOS_Bool EOS_MEMORY_CALL Decommit(void* Address, size_t SizeInBytes, void* ContextData);
+}
diff --git a/Engine/Plugins/Online/EOSShared/Source/EOSShared/Private/Windows/WindowsEOSSDKManager.cpp b/Engine/Plugins/Online/EOSShared/Source/EOSShared/Private/Windows/WindowsEOSSDKManager.cpp
--- a/Engine/Plugins/Online/EOSShared/Source/EOSShared/Private/Windows/WindowsEOSSDKManager.cpp
+++ b/Engine/Plugins/Online/EOSShared/Source/EOSShared/Private/Windows/WindowsEOSSDKManager.cpp
@@ -11,15 +11,14 @@
 
 FWindowsEOSSDKManager::FWindowsEOSSDKManager()
 {
-	PlatformSteamOptions =
-	{
-		2, //EOS_INTEGRATEDPLATFORM_STEAM_OPTIONS_API_LATEST
-		nullptr,
-		1,
-		57
-	};
-
-	UE_EOS_CHECK_API_MISMATCH(EOS_INTEGRATEDPLATFORM_STEAM_OPTIONS_API_LATEST, 2);
+	PlatformSteamOptions.ApiVersion = 3; //EOS_INTEGRATEDPLATFORM_STEAM_OPTIONS_API_LATEST;
+	PlatformSteamOptions.OverrideLibraryPath = nullptr;
+	PlatformSteamOptions.SteamMajorVersion = 1;
+	PlatformSteamOptions.SteamMinorVersion = 57;
+	PlatformSteamOptions.SteamApiInterfaceVersionsArray = nullptr;
+	PlatformSteamOptions.SteamApiInterfaceVersionsArrayBytes = 0;
+
+	UE_EOS_CHECK_API_MISMATCH(EOS_INTEGRATEDPLATFORM_STEAM_OPTIONS_API_LATEST, 3);
 }
 
 FWindowsEOSSDKManager::~FWindowsEOSSDKManager()
diff --git a/Engine/Plugins/Online/OnlineSubsystemEOS/Source/OnlineSubsystemEOS/Private/OnlineSubsystemEOS.cpp b/Engine/Plugins/Online/OnlineSubsystemEOS/Source/OnlineSubsystemEOS/Private/OnlineSubsystemEOS.cpp
--- a/Engine/Plugins/Online/OnlineSubsystemEOS/Source/OnlineSubsystemEOS/Private/OnlineSubsystemEOS.cpp
+++ b/Engine/Plugins/Online/OnlineSubsystemEOS/Source/OnlineSubsystemEOS/Private/OnlineSubsystemEOS.cpp
@@ -189,8 +189,8 @@ bool FOnlineSubsystemEOS::PlatformCreate()
 
 	// Create platform instance
 	EOS_Platform_Options PlatformOptions = {};
-	PlatformOptions.ApiVersion = 13;
-	UE_EOS_CHECK_API_MISMATCH(EOS_PLATFORM_OPTIONS_API_LATEST, 13);
+	PlatformOptions.ApiVersion = 14;
+	UE_EOS_CHECK_API_MISMATCH(EOS_PLATFORM_OPTIONS_API_LATEST, 14);
 	PlatformOptions.ClientCredentials.ClientId = reinterpret_cast<const char*>(ArtifactSettings.ClientId.IsEmpty() ? nullptr : (const char*)ClientIdUtf8.Get());
 	PlatformOptions.ClientCredentials.ClientSecret = ArtifactSettings.ClientSecret.IsEmpty() ? nullptr : (const char*)ClientSecretUtf8.Get();
 	PlatformOptions.ProductId = ArtifactSettings.ProductId.IsEmpty() ? nullptr : (const char*)ProductIdUtf8.Get();
@@ -199,6 +199,7 @@ bool FOnlineSubsystemEOS::PlatformCreate()
 	PlatformOptions.EncryptionKey = ArtifactSettings.EncryptionKey.IsEmpty() ? nullptr : (const char*)EncryptionKeyUtf8.Get();
 	PlatformOptions.bIsServer = IsRunningDedicatedServer() ? EOS_TRUE : EOS_FALSE;
 	PlatformOptions.Reserved = nullptr;
+	PlatformOptions.TaskNetworkTimeoutSeconds = nullptr;
 
 	FEOSSettings EOSSettings = UEOSSettings::GetSettings();
 	uint64 OverlayFlags = 0;
diff --git a/Engine/Plugins/Online/VoiceChat/EOSVoiceChat/Source/EOSVoiceChat/Private/EOSVoiceChat.cpp b/Engine/Plugins/Online/VoiceChat/EOSVoiceChat/Source/EOSVoiceChat/Private/EOSVoiceChat.cpp
--- a/Engine/Plugins/Online/VoiceChat/EOSVoiceChat/Source/EOSVoiceChat/Private/EOSVoiceChat.cpp
+++ b/Engine/Plugins/Online/VoiceChat/EOSVoiceChat/Source/EOSVoiceChat/Private/EOSVoiceChat.cpp
@@ -132,8 +132,8 @@ void FEOSVoiceChat::Initialize(const FOnVoiceChatInitializeCompleteDelegate& Ini
 					const FTCHARToUTF8 Utf8OverrideLocaleCode(*ConfigOverrideLocaleCode);
 
 					EOS_Platform_Options PlatformOptions = {};
-					PlatformOptions.ApiVersion = 13;
-					UE_EOS_CHECK_API_MISMATCH(EOS_PLATFORM_OPTIONS_API_LATEST, 13);
+					PlatformOptions.ApiVersion = 14;
+					UE_EOS_CHECK_API_MISMATCH(EOS_PLATFORM_OPTIONS_API_LATEST, 14);
 					PlatformOptions.Reserved = nullptr;
 					PlatformOptions.SystemSpecificOptions = nullptr;
 					PlatformOptions.ProductId = ConfigProductId.IsEmpty() ? nullptr : Utf8ProductId.Get();
@@ -149,6 +149,7 @@ void FEOSVoiceChat::Initialize(const FOnVoiceChatInitializeCompleteDelegate& Ini
 					PlatformOptions.CacheDirectory = nullptr;
 					PlatformOptions.TickBudgetInMilliseconds = 1;
 					PlatformOptions.IntegratedPlatformOptionsContainerHandle = nullptr;
+					PlatformOptions.TaskNetworkTimeoutSeconds = nullptr;
 #if UE_EDITOR
 					//PlatformCreateOptions.Flags |= EOS_PF_LOADING_IN_EDITOR;
 #endif
diff --git a/Engine/Source/Runtime/Core/Public/Windows/WindowsPlatformMemory.h b/Engine/Source/Runtime/Core/Public/Windows/WindowsPlatformMemory.h
--- a/Engine/Source/Runtime/Core/Public/Windows/WindowsPlatformMemory.h
+++ b/Engine/Source/Runtime/Core/Public/Windows/WindowsPlatformMemory.h
@@ -115,9 +115,9 @@ struct FWindowsPlatformMemory
 		FPlatformVirtualMemoryBlock(const FPlatformVirtualMemoryBlock& Other) = default;
 		FPlatformVirtualMemoryBlock& operator=(const FPlatformVirtualMemoryBlock& Other) = default;
 
-		void Commit(size_t InOffset, size_t InSize);
-		void Decommit(size_t InOffset, size_t InSize);
-		void FreeVirtual();
+		CORE_API void Commit(size_t InOffset, size_t InSize);
+		CORE_API void Decommit(size_t InOffset, size_t InSize);
+		CORE_API void FreeVirtual();
 
 		FORCEINLINE void CommitByPtr(void *InPtr, size_t InSize)
 		{
@@ -144,9 +144,9 @@ struct FWindowsPlatformMemory
 			return VMSizeDivVirtualSizeAlignment * GetVirtualSizeAlignment();
 		}
 
-		static FPlatformVirtualMemoryBlock AllocateVirtual(size_t Size, size_t InAlignment = FPlatformVirtualMemoryBlock::GetVirtualSizeAlignment());
-		static size_t GetCommitAlignment();
-		static size_t GetVirtualSizeAlignment();
+		static CORE_API FPlatformVirtualMemoryBlock AllocateVirtual(size_t Size, size_t InAlignment = FPlatformVirtualMemoryBlock::GetVirtualSizeAlignment());
+		static CORE_API size_t GetCommitAlignment();
+		static CORE_API size_t GetVirtualSizeAlignment();
 	};
 
 	static CORE_API FSharedMemoryRegion* MapNamedSharedMemoryRegion(const FString& InName, bool bCreate, uint32 AccessMode, SIZE_T Size, const void* pSecurityAttributes = nullptr);
