- Add missing WITH_EOS_SDK guards

diff --git a/Engine/Plugins/Online/EOSShared/Source/EOSShared/Private/EOSVirtualMemoryCallbacks.cpp b/Engine/Plugins/Online/EOSShared/Source/EOSShared/Private/EOSVirtualMemoryCallbacks.cpp
--- a/Engine/Plugins/Online/EOSShared/Source/EOSShared/Private/EOSVirtualMemoryCallbacks.cpp
+++ b/Engine/Plugins/Online/EOSShared/Source/EOSShared/Private/EOSVirtualMemoryCallbacks.cpp
@@ -2,6 +2,8 @@
 
 #include "EOSVirtualMemoryCallbacks.h"
 
+#if WITH_EOS_SDK
+
 #include "HAL/PlatformMemory.h"
 
 namespace EOSVirtualMemoryCallbacks
@@ -37,3 +39,5 @@ EOS_Bool EOS_MEMORY_CALL Decommit(void* Address, size_t SizeInBytes, void* Conte
 }
 
 } // namespace EOSVirtualMemoryCallbacks
+
+#endif // WITH_EOS_SDK
\ No newline at end of file
diff --git a/Engine/Plugins/Online/EOSShared/Source/EOSShared/Private/EOSVirtualMemoryCallbacks.h b/Engine/Plugins/Online/EOSShared/Source/EOSShared/Private/EOSVirtualMemoryCallbacks.h
--- a/Engine/Plugins/Online/EOSShared/Source/EOSShared/Private/EOSVirtualMemoryCallbacks.h
+++ b/Engine/Plugins/Online/EOSShared/Source/EOSShared/Private/EOSVirtualMemoryCallbacks.h
@@ -2,6 +2,8 @@
 
 #pragma once
 
+#if WITH_EOS_SDK
+
 #if defined(EOS_PLATFORM_BASE_FILE_NAME)
 #include EOS_PLATFORM_BASE_FILE_NAME
 #endif
@@ -17,3 +19,5 @@ namespace EOSVirtualMemoryCallbacks
 	EOS_Bool EOS_MEMORY_CALL Commit(void* Address, size_t SizeInBytes, void* ContextData);
 	EOS_Bool EOS_MEMORY_CALL Decommit(void* Address, size_t SizeInBytes, void* ContextData);
 }
+
+#endif // WITH_EOS_SDK
\ No newline at end of file
