Parse relying party URI from config.

diff --git a/Engine/Plugins/Online/EOSShared/Source/EOSShared/Private/EOSSDKManager.cpp b/Engine/Plugins/Online/EOSShared/Source/EOSShared/Private/EOSSDKManager.cpp
--- a/Engine/Plugins/Online/EOSShared/Source/EOSShared/Private/EOSSDKManager.cpp
+++ b/Engine/Plugins/Online/EOSShared/Source/EOSShared/Private/EOSSDKManager.cpp
@@ -336,6 +336,7 @@ const FEOSSDKPlatformConfig* FEOSSDKManager::GetPlatformConfig(const FString& Pl
 		UE_LOG(LogEOSSDK, Warning, TEXT("Config section \"EOSSDK.Platform.%s\" contains deprecated key EncryptionKey, please migrate to ClientEncryptionKey."), *PlatformConfigName);
 	}
 	GConfig->GetString(*SectionName, TEXT("ClientEncryptionKey"), PlatformConfig->EncryptionKey, GEngineIni);
+	GConfig->GetString(*SectionName, TEXT("RelyingPartyURI"), PlatformConfig->RelyingPartyURI, GEngineIni);
 	GConfig->GetString(*SectionName, TEXT("OverrideCountryCode"), PlatformConfig->OverrideCountryCode, GEngineIni);
 	GConfig->GetString(*SectionName, TEXT("OverrideLocaleCode"), PlatformConfig->OverrideLocaleCode, GEngineIni);
 	GConfig->GetString(*SectionName, TEXT("DeploymentId"), PlatformConfig->DeploymentId, GEngineIni);
