#EOS Extend EOS SDK log suppression system to handle log categories

diff --git a/Engine/Plugins/Online/EOSShared/Source/EOSShared/Private/EOSSDKManager.cpp b/Engine/Plugins/Online/EOSShared/Source/EOSShared/Private/EOSSDKManager.cpp
--- a/Engine/Plugins/Online/EOSShared/Source/EOSShared/Private/EOSSDKManager.cpp
+++ b/Engine/Plugins/Online/EOSShared/Source/EOSShared/Private/EOSSDKManager.cpp
@@ -102,6 +102,18 @@ namespace
 			if(FEOSSharedModule* Module = FEOSSharedModule::Get())
 			{
 				bSuppressLogLevel = Algo::AnyOf(Module->GetSuppressedLogStrings(), [&MessageStr](const FString& SuppressedLogString) { return MessageStr.Contains(SuppressedLogString); });
+
+				// Check for suppressed categories if not already suppressed
+				if (!bSuppressLogLevel)
+				{
+					const TArray<FString>& SuppressedCategories = Module->GetSuppressedLogCategories();
+					if (SuppressedCategories.Num())
+					{
+						FString CategoryStr(UTF8_TO_TCHAR(Message->Category));
+						CategoryStr.TrimStartAndEndInline();
+						bSuppressLogLevel = Algo::AnyOf(SuppressedCategories, [&CategoryStr](const FString& SuppressedLogCategory) { return CategoryStr.Contains(SuppressedLogCategory); });
+					}
+				}
 			}
 		}
 
diff --git a/Engine/Plugins/Online/EOSShared/Source/EOSShared/Private/EOSSharedModule.cpp b/Engine/Plugins/Online/EOSShared/Source/EOSShared/Private/EOSSharedModule.cpp
--- a/Engine/Plugins/Online/EOSShared/Source/EOSShared/Private/EOSSharedModule.cpp
+++ b/Engine/Plugins/Online/EOSShared/Source/EOSShared/Private/EOSSharedModule.cpp
@@ -101,6 +101,7 @@ void FEOSSharedModule::OnConfigSectionsChanged(const FString& IniFilename, const
 void FEOSSharedModule::LoadConfig()
 {
 	GConfig->GetArray(CONFIG_SECTION_NAME, TEXT("SuppressedLogStrings"), SuppressedLogStrings, GEngineIni);
+	GConfig->GetArray(CONFIG_SECTION_NAME, TEXT("SuppressedLogCategories"), SuppressedLogCategories, GEngineIni);
 }
 
 #undef CONFIG_SECTION_NAME
diff --git a/Engine/Plugins/Online/EOSShared/Source/EOSShared/Private/EOSSharedModule.h b/Engine/Plugins/Online/EOSShared/Source/EOSShared/Private/EOSSharedModule.h
--- a/Engine/Plugins/Online/EOSShared/Source/EOSShared/Private/EOSSharedModule.h
+++ b/Engine/Plugins/Online/EOSShared/Source/EOSShared/Private/EOSSharedModule.h
@@ -16,6 +16,7 @@ public:
 	static FEOSSharedModule* Get();
 
 	const TArray<FString>& GetSuppressedLogStrings() const { return SuppressedLogStrings; }
+	const TArray<FString>& GetSuppressedLogCategories() const { return SuppressedLogCategories; }
 
 private:
 	// ~Begin IModuleInterface
@@ -31,5 +32,6 @@ private:
 	TUniquePtr<FEOSSDKManager> SDKManager;
 #endif
 	TArray<FString> SuppressedLogStrings;
+	TArray<FString> SuppressedLogCategories;
 	FDelegateHandle OnPostForkDelegateHandle;
 };
\ No newline at end of file
