Enable lazy-loading and early-unloading of external module DLLs when the overlay is only used for the login flow

diff --git a/Engine/Plugins/Online/EOSShared/Source/EOSShared/Private/EOSSDKManager.cpp b/Engine/Plugins/Online/EOSShared/Source/EOSShared/Private/EOSSDKManager.cpp
--- a/Engine/Plugins/Online/EOSShared/Source/EOSShared/Private/EOSSDKManager.cpp
+++ b/Engine/Plugins/Online/EOSShared/Source/EOSShared/Private/EOSSDKManager.cpp
@@ -643,6 +643,8 @@ IEOSPlatformHandlePtr FEOSSDKManager::CreatePlatform(EOS_Platform_Options& Platf
 
 		OnPreCreatePlatform.Broadcast(PlatformOptions);
 
+		ApplyOverlayPlatformOptions(PlatformOptions);
+
 		const EOS_HPlatform PlatformHandle = EOS_Platform_Create(&PlatformOptions);
 		if (PlatformHandle)
 		{
@@ -850,6 +852,16 @@ void FEOSSDKManager::SetInvokeOverlayButton(const EOS_HPlatform PlatformHandle)
 	}
 }
 
+void FEOSSDKManager::ApplyOverlayPlatformOptions(EOS_Platform_Options& PlatformOptions)
+{
+	// On consoles, if we're only using the overlay for the AP login flow we can enable auto loading/unloading
+	// so that we only load the KITT DLLs when we open a browser and we unload them after closing the browser
+	if (PlatformOptions.Flags & EOS_PF_DISABLE_SOCIAL_OVERLAY)
+	{
+		PlatformOptions.Flags |= EOS_PF_CONSOLE_ENABLE_OVERLAY_AUTOMATIC_UNLOADING;
+	}
+}
+
 bool FEOSSDKManager::Tick(float)
 {
 	check(IsInGameThread());
diff --git a/Engine/Plugins/Online/EOSShared/Source/EOSShared/Private/EOSSDKManager.h b/Engine/Plugins/Online/EOSShared/Source/EOSShared/Private/EOSSDKManager.h
--- a/Engine/Plugins/Online/EOSShared/Source/EOSShared/Private/EOSSDKManager.h
+++ b/Engine/Plugins/Online/EOSShared/Source/EOSShared/Private/EOSSDKManager.h
@@ -92,6 +92,7 @@ protected:
 #endif
 
 	void SetInvokeOverlayButton(const EOS_HPlatform PlatformHandle);
+	void ApplyOverlayPlatformOptions(EOS_Platform_Options& PlatformOptions);
 	EOS_HIntegratedPlatformOptionsContainer CreateIntegratedPlatformOptionsContainer();
 	void ApplyIntegratedPlatformOptions(EOS_HIntegratedPlatformOptionsContainer& Container);
 	virtual void ApplySystemSpecificOptions(const void*& SystemSpecificOptions);
