Add OnPostInitializeSDK event to IEOSSDKManager

diff --git a/Engine/Plugins/Online/EOSShared/Source/EOSShared/Private/EOSSDKManager.cpp b/Engine/Plugins/Online/EOSShared/Source/EOSShared/Private/EOSSDKManager.cpp
--- a/Engine/Plugins/Online/EOSShared/Source/EOSShared/Private/EOSSDKManager.cpp
+++ b/Engine/Plugins/Online/EOSShared/Source/EOSShared/Private/EOSSDKManager.cpp
@@ -309,6 +309,7 @@ EOS_EResult FEOSSDKManager::Initialize()
 			UE_LOG(LogEOSSDK, Warning, TEXT("EOS_Initialize failed error:%s"), *LexToString(EosResult));
 		}
 
+		OnPostInitializeSDK.Broadcast(EosResult);
 		return EosResult;
 	}
 }
diff --git a/Engine/Plugins/Online/EOSShared/Source/EOSShared/Public/IEOSSDKManager.h b/Engine/Plugins/Online/EOSShared/Source/EOSShared/Public/IEOSSDKManager.h
--- a/Engine/Plugins/Online/EOSShared/Source/EOSShared/Public/IEOSSDKManager.h
+++ b/Engine/Plugins/Online/EOSShared/Source/EOSShared/Public/IEOSSDKManager.h
@@ -40,6 +40,7 @@ struct FEOSSDKPlatformConfig
 };
 
 DECLARE_MULTICAST_DELEGATE_OneParam(FEOSSDKManagerOnPreInitializeSDK, EOS_InitializeOptions& Options);
+DECLARE_MULTICAST_DELEGATE_OneParam(FEOSSDKManagerOnPostInitializeSDK, EOS_EResult Result);
 DECLARE_MULTICAST_DELEGATE_TwoParams(FEOSSDKManagerOnDefaultPlatformConfigNameChanged, const FString& NewName, const FString& OldName);
 DECLARE_MULTICAST_DELEGATE_TwoParams(FEOSSDKManagerOnPreCreateNamedPlatform, const FEOSSDKPlatformConfig& Config, EOS_Platform_Options& Options);
 DECLARE_MULTICAST_DELEGATE_OneParam(FEOSSDKManagerOnPreCreatePlatform, EOS_Platform_Options& Options);
@@ -134,6 +135,7 @@ public:
 	virtual void AddCallbackObject(TUniquePtr<class FCallbackBase> CallbackObj) = 0;
 
 	FEOSSDKManagerOnPreInitializeSDK OnPreInitializeSDK;
+	FEOSSDKManagerOnPostInitializeSDK OnPostInitializeSDK;
 	FEOSSDKManagerOnDefaultPlatformConfigNameChanged OnDefaultPlatformConfigNameChanged;
 	FEOSSDKManagerOnPreCreateNamedPlatform OnPreCreateNamedPlatform;
 	FEOSSDKManagerOnPreCreatePlatform OnPreCreatePlatform;
