Allow completing EOS SDK login through external browser redirection on iOS

diff --git a/Engine/Plugins/Online/EOSShared/Source/EOSShared/Private/IOS/IOSEOSSDKManager.cpp b/Engine/Plugins/Online/EOSShared/Source/EOSShared/Private/IOS/IOSEOSSDKManager.cpp
--- a/Engine/Plugins/Online/EOSShared/Source/EOSShared/Private/IOS/IOSEOSSDKManager.cpp
+++ b/Engine/Plugins/Online/EOSShared/Source/EOSShared/Private/IOS/IOSEOSSDKManager.cpp
@@ -3,12 +3,20 @@
 
 #if WITH_EOS_SDK
 
+#include "IOSAppDelegate.h"
 #include "Misc/CoreDelegates.h"
 
+static void OnUrlOpened(UIApplication* application, NSURL* url, NSString* sourceApplication, id annotation)
+{
+	// TODO: This is based on a prototype fix on EOS SDK. Once the fix is properly submitted to EOS SDK we should update it
+	[[NSNotificationCenter defaultCenter] postNotificationName:@"EOSSDKAuthCallbackNotification" object:nil userInfo: @{@"EOSSDKAuthCallbackURLKey" : url}];
+}
+
 FIOSEOSSDKManager::FIOSEOSSDKManager()
 {
 	FCoreDelegates::ApplicationHasEnteredForegroundDelegate.AddRaw(this, &FIOSEOSSDKManager::OnApplicationStatusChanged, EOS_EApplicationStatus::EOS_AS_Foreground);
 	FCoreDelegates::ApplicationWillEnterBackgroundDelegate.AddRaw(this, &FIOSEOSSDKManager::OnApplicationStatusChanged, EOS_EApplicationStatus::EOS_AS_BackgroundSuspended);
+	FIOSCoreDelegates::OnOpenURL.AddStatic(&OnUrlOpened);
 }
 
 FIOSEOSSDKManager::~FIOSEOSSDKManager()
@@ -27,4 +35,4 @@ FString FIOSEOSSDKManager::GetCacheDirBase() const
 	return FString(UTF8_TO_TCHAR(CStrCacheDirectory));
 };
 
-#endif // WITH_EOS_SDK
\ No newline at end of file
+#endif // WITH_EOS_SDK
