Fix missing dependency

diff --git a/Engine/Plugins/Online/EOSShared/Source/EOSShared/EOSShared.Build.cs b/Engine/Plugins/Online/EOSShared/Source/EOSShared/EOSShared.Build.cs
--- a/Engine/Plugins/Online/EOSShared/Source/EOSShared/EOSShared.Build.cs
+++ b/Engine/Plugins/Online/EOSShared/Source/EOSShared/EOSShared.Build.cs
@@ -33,6 +33,11 @@ public class EOSShared : ModuleRules
 			}
 		);
 
+		if (Target.Platform == UnrealTargetPlatform.IOS)
+		{
+			PrivateDependencyModuleNames.Add("ApplicationCore");
+		}
+		
 		if (Target.bCompileAgainstEngine)
 		{
 			PrivateDependencyModuleNames.Add("Slate");
