- Fix incorrect log string in LexToString, and move it to EOSShared

diff --git a/Engine/Plugins/Online/EOSShared/Source/EOSShared/Private/EOSShared.cpp b/Engine/Plugins/Online/EOSShared/Source/EOSShared/Private/EOSShared.cpp
--- a/Engine/Plugins/Online/EOSShared/Source/EOSShared/Private/EOSShared.cpp
+++ b/Engine/Plugins/Online/EOSShared/Source/EOSShared/Private/EOSShared.cpp
@@ -9,10 +9,11 @@
 
 #include "eos_auth_types.h"
 #include "eos_friends_types.h"
+#include "eos_p2p_types.h"
 #include "eos_presence_types.h"
+#include "eos_rtc_audio_types.h"
 #include "eos_rtc_types.h"
 #include "eos_userinfo_types.h"
-#include "eos_p2p_types.h"
 
 DEFINE_LOG_CATEGORY(LogEOSSDK);
 
@@ -204,6 +205,19 @@ const TCHAR* LexToString(const EOS_EExternalCredentialType ExternalCredentialTyp
 	}
 }
 
+const TCHAR* LexToString(const EOS_ERTCAudioInputStatus Status)
+{
+	switch (Status)
+	{
+		case EOS_ERTCAudioInputStatus::EOS_RTCAIS_Idle:						return TEXT("EOS_RTCAIS_Idle");
+		case EOS_ERTCAudioInputStatus::EOS_RTCAIS_Recording: 				return TEXT("EOS_RTCAIS_Recording");
+		case EOS_ERTCAudioInputStatus::EOS_RTCAIS_RecordingSilent: 			return TEXT("EOS_RTCAIS_RecordingSilent");
+		case EOS_ERTCAudioInputStatus::EOS_RTCAIS_RecordingDisconnected:	return TEXT("EOS_RTCAIS_RecordingDisconnected");
+		case EOS_ERTCAudioInputStatus::EOS_RTCAIS_Failed:					return TEXT("EOS_RTCAIS_Failed");
+		default: checkNoEntry();											return TEXT("Unknown");
+	}
+}
+
 bool LexFromString(EOS_EExternalAccountType& OutEnum, const TCHAR* InString)
 {
 	if (FCString::Stricmp(InString, TEXT("Amazon")) == 0)
diff --git a/Engine/Plugins/Online/EOSShared/Source/EOSShared/Public/EOSShared.h b/Engine/Plugins/Online/EOSShared/Source/EOSShared/Public/EOSShared.h
--- a/Engine/Plugins/Online/EOSShared/Source/EOSShared/Public/EOSShared.h
+++ b/Engine/Plugins/Online/EOSShared/Source/EOSShared/Public/EOSShared.h
@@ -52,9 +52,10 @@ EOS_ENUM_FORWARD_DECL(EOS_EDesktopCrossplayStatus);
 EOS_ENUM_FORWARD_DECL(EOS_EFriendsStatus);
 EOS_ENUM_FORWARD_DECL(EOS_ELoginCredentialType);
 EOS_ENUM_FORWARD_DECL(EOS_ENetworkStatus);
+EOS_ENUM_FORWARD_DECL(EOS_EPacketReliability);
+EOS_ENUM_FORWARD_DECL(EOS_ERTCAudioInputStatus);
 EOS_ENUM_FORWARD_DECL(EOS_Presence_EStatus);
 EOS_ENUM_FORWARD_DECL(EOS_UI_EInputStateButtonFlags);
-EOS_ENUM_FORWARD_DECL(EOS_EPacketReliability);
 #undef EOS_ENUM_FORWARD_DECL
 
 #define EOS_STRUCT_FORWARD_DECL(name) extern "C" typedef struct _tag ## name name;
@@ -88,6 +89,7 @@ EOSSHARED_API const TCHAR* LexToString(const EOS_EExternalAccountType ExternalAc
 EOSSHARED_API const TCHAR* LexToString(const EOS_EFriendsStatus FriendStatus);
 EOSSHARED_API const TCHAR* LexToString(const EOS_ELoginStatus LoginStatus);
 EOSSHARED_API const TCHAR* LexToString(const EOS_ENetworkStatus NetworkStatus);
+EOSSHARED_API const TCHAR* LexToString(const EOS_ERTCAudioInputStatus Status);
 EOSSHARED_API const TCHAR* LexToString(const EOS_Presence_EStatus PresenceStatus);
 EOSSHARED_API const TCHAR* LexToString(const EOS_EExternalCredentialType ExternalCredentialType);
 
diff --git a/Engine/Plugins/Online/VoiceChat/EOSVoiceChat/Source/EOSVoiceChat/Private/EOSVoiceChat.cpp b/Engine/Plugins/Online/VoiceChat/EOSVoiceChat/Source/EOSVoiceChat/Private/EOSVoiceChat.cpp
--- a/Engine/Plugins/Online/VoiceChat/EOSVoiceChat/Source/EOSVoiceChat/Private/EOSVoiceChat.cpp
+++ b/Engine/Plugins/Online/VoiceChat/EOSVoiceChat/Source/EOSVoiceChat/Private/EOSVoiceChat.cpp
@@ -28,19 +28,6 @@
 
 DEFINE_LOG_CATEGORY(LogEOSVoiceChat);
 
-const TCHAR* LexToString(EOS_ERTCAudioInputStatus Status)
-{
-	switch (Status)
-	{
-	case EOS_ERTCAudioInputStatus::EOS_RTCAIS_Idle: return TEXT("EOS_RTCAIS_Idle");
-	case EOS_ERTCAudioInputStatus::EOS_RTCAIS_Recording: return TEXT("EOS_RTCAIS_Recording");
-	case EOS_ERTCAudioInputStatus::EOS_RTCAIS_RecordingSilent: return TEXT("EOS_RTCAIS_RecordingSilent");
-	case EOS_ERTCAudioInputStatus::EOS_RTCAIS_RecordingDisconnected: return TEXT("EOS_RTCAIS_RecordingDisconnected");
-	case EOS_ERTCAudioInputStatus::EOS_RTCAIS_Failed: return TEXT("EOS_RTCEOS_RTCAIS_Failed_AudioInputFailed");
-	default: return TEXT("Unknown");
-	}
-}
-
 FEOSVoiceChatDelegates::FOnAudioInputDeviceStatusChanged FEOSVoiceChatDelegates::OnAudioInputDeviceStatusChanged;
 FEOSVoiceChatDelegates::FOnVoiceChatPlayerAddedMetadataDelegate FEOSVoiceChatDelegates::OnVoiceChatPlayerAddedMetadata;
 FEOSVoiceChatDelegates::FOnAudioStatusChanged FEOSVoiceChatDelegates::OnAudioStatusChanged;
diff --git a/Engine/Plugins/Online/VoiceChat/EOSVoiceChat/Source/EOSVoiceChat/Public/EOSVoiceChat.h b/Engine/Plugins/Online/VoiceChat/EOSVoiceChat/Source/EOSVoiceChat/Public/EOSVoiceChat.h
--- a/Engine/Plugins/Online/VoiceChat/EOSVoiceChat/Source/EOSVoiceChat/Public/EOSVoiceChat.h
+++ b/Engine/Plugins/Online/VoiceChat/EOSVoiceChat/Source/EOSVoiceChat/Public/EOSVoiceChat.h
@@ -18,8 +18,6 @@
 #include "eos_rtc_audio_types.h"
 #include "eos_lobby_types.h"
 
-EOSVOICECHAT_API const TCHAR* LexToString(EOS_ERTCAudioInputStatus Status);
-
 struct FVoiceChatMetadataItem
 {
 	FString Key;
