Unshelved from pending changelist '34339324':

Increase the EOSSDK Tick rate when the overlay is visible.

The overlay input processing and some of the UI elements are tied to the Tick, so we need it to be increased when the overlay becomes visible, as the EOSSDK Tick rate is limited by default for perf reasons.

#rnx

diff --git a/Engine/Plugins/Online/EOSShared/Source/EOSShared/Private/EOSSDKManager.cpp b/Engine/Plugins/Online/EOSShared/Source/EOSShared/Private/EOSSDKManager.cpp
--- a/Engine/Plugins/Online/EOSShared/Source/EOSShared/Private/EOSSDKManager.cpp
+++ b/Engine/Plugins/Online/EOSShared/Source/EOSShared/Private/EOSSDKManager.cpp
@@ -663,6 +663,7 @@ IEOSPlatformHandlePtr FEOSSDKManager::CreatePlatform(EOS_Platform_Options& Platf
 			EOS_Platform_SetNetworkStatus(PlatformHandle, ConvertNetworkStatus(FPlatformMisc::GetNetworkConnectionStatus()));
 			
 			SetInvokeOverlayButton(PlatformHandle);
+			RegisterDisplaySettingsUpdatedCallback(PlatformHandle);
 
 			// Tick the platform once to work around EOSSDK error logging that occurs if you create then immediately destroy a platform.
 			SharedPlatform->Tick();
@@ -749,7 +750,7 @@ void FEOSSDKManager::LoadConfig()
 	SetupTicker();
 }
 
-void FEOSSDKManager::SetupTicker()
+void FEOSSDKManager::SetupTicker(bool bIgnoreConfigTickInterval)
 {
 	check(IsInGameThread());
 
@@ -762,7 +763,7 @@ void FEOSSDKManager::SetupTicker()
 	int NumActivePlatforms = ActivePlatforms.Num();
 	if (NumActivePlatforms > 0)
 	{
-		const double TickIntervalSeconds = ConfigTickIntervalSeconds > SMALL_NUMBER ? ConfigTickIntervalSeconds / NumActivePlatforms : 0.f;
+		const double TickIntervalSeconds = ConfigTickIntervalSeconds > SMALL_NUMBER && !bIgnoreConfigTickInterval ? ConfigTickIntervalSeconds / NumActivePlatforms : 0.f;
 		TickerHandle = FTSTicker::GetCoreTicker().AddTicker(FTickerDelegate::CreateRaw(this, &FEOSSDKManager::Tick), TickIntervalSeconds);
 	}
 }
@@ -853,6 +854,26 @@ void FEOSSDKManager::SetInvokeOverlayButton(const EOS_HPlatform PlatformHandle)
 	}
 }
 
+void FEOSSDKManager::OnDisplaySettingsUpdated(const EOS_UI_OnDisplaySettingsUpdatedCallbackInfo* Data)
+{
+	reinterpret_cast<FEOSSDKManager*>(Data->ClientData)->SetupTicker((bool)Data->bIsVisible);
+}
+
+void FEOSSDKManager::RegisterDisplaySettingsUpdatedCallback(const EOS_HPlatform PlatformHandle)
+{
+	if (bEnablePlatformIntegration)
+	{
+		if (EOS_HUI UIHandle = EOS_Platform_GetUIInterface(PlatformHandle))
+		{
+			EOS_UI_AddNotifyDisplaySettingsUpdatedOptions Options = {};
+			Options.ApiVersion = 1;
+			UE_EOS_CHECK_API_MISMATCH(EOS_UI_ADDNOTIFYDISPLAYSETTINGSUPDATED_API_LATEST, 1);
+
+			EOS_NotificationId DisplaySettingsUpdatedId = EOS_UI_AddNotifyDisplaySettingsUpdated(UIHandle, &Options, this, &FEOSSDKManager::OnDisplaySettingsUpdated);
+		}
+	}
+}
+
 void FEOSSDKManager::ApplyOverlayPlatformOptions(EOS_Platform_Options& PlatformOptions)
 {
 	// On consoles, if we're only using the overlay for the AP login flow we can enable auto loading/unloading
diff --git a/Engine/Plugins/Online/EOSShared/Source/EOSShared/Private/EOSSDKManager.h b/Engine/Plugins/Online/EOSShared/Source/EOSShared/Private/EOSSDKManager.h
--- a/Engine/Plugins/Online/EOSShared/Source/EOSShared/Private/EOSSDKManager.h
+++ b/Engine/Plugins/Online/EOSShared/Source/EOSShared/Private/EOSSDKManager.h
@@ -91,6 +91,9 @@ protected:
 	virtual bool IsRenderReady();
 #endif
 
+	static void OnDisplaySettingsUpdated(const EOS_UI_OnDisplaySettingsUpdatedCallbackInfo* Data);
+	void RegisterDisplaySettingsUpdatedCallback(const EOS_HPlatform PlatformHandle);
+
 	void SetInvokeOverlayButton(const EOS_HPlatform PlatformHandle);
 	void ApplyOverlayPlatformOptions(EOS_Platform_Options& PlatformOptions);
 	EOS_HIntegratedPlatformOptionsContainer CreateIntegratedPlatformOptionsContainer();
@@ -109,7 +112,7 @@ protected:
 	void LoadConfig();
 	void ReleasePlatform(EOS_HPlatform PlatformHandle);
 	void ReleaseReleasedPlatforms();
-	void SetupTicker();
+	void SetupTicker(bool bIgnoreConfigTickInterval = false);
 	void OnLogVerbosityChanged(const FLogCategoryName& CategoryName, ELogVerbosity::Type OldVerbosity, ELogVerbosity::Type NewVerbosity);
 
 #if EOSSDK_RUNTIME_LOAD_REQUIRED
