Added logic to handle overwriting of EOS Platform Config after MCP backend changes.

diff --git a/Engine/Plugins/Online/EOSShared/Source/EOSShared/Private/EOSSDKManager.cpp b/Engine/Plugins/Online/EOSShared/Source/EOSShared/Private/EOSSDKManager.cpp
--- a/Engine/Plugins/Online/EOSShared/Source/EOSShared/Private/EOSSDKManager.cpp
+++ b/Engine/Plugins/Online/EOSShared/Source/EOSShared/Private/EOSSDKManager.cpp
@@ -398,7 +398,7 @@ const FEOSSDKPlatformConfig* FEOSSDKManager::GetPlatformConfig(const FString& Pl
 	return PlatformConfig;
 }
 
-bool FEOSSDKManager::AddPlatformConfig(const FEOSSDKPlatformConfig& PlatformConfig)
+bool FEOSSDKManager::AddPlatformConfig(const FEOSSDKPlatformConfig& PlatformConfig, bool bOverwriteExistingConfig)
 {
 	if (PlatformConfig.Name.IsEmpty())
 	{
@@ -406,7 +406,7 @@ bool FEOSSDKManager::AddPlatformConfig(const FEOSSDKPlatformConfig& PlatformConf
 		return false;
 	}
 
-	if (PlatformConfigs.Find(PlatformConfig.Name))
+	if (PlatformConfigs.Find(PlatformConfig.Name) && !bOverwriteExistingConfig)
 	{
 		UE_LOG(LogEOSSDK, Warning, TEXT("Platform config already exists: %s"), *PlatformConfig.Name);
 		return false;
diff --git a/Engine/Plugins/Online/EOSShared/Source/EOSShared/Private/EOSSDKManager.h b/Engine/Plugins/Online/EOSShared/Source/EOSShared/Private/EOSSDKManager.h
--- a/Engine/Plugins/Online/EOSShared/Source/EOSShared/Private/EOSSDKManager.h
+++ b/Engine/Plugins/Online/EOSShared/Source/EOSShared/Private/EOSSDKManager.h
@@ -43,7 +43,7 @@ public:
 	virtual bool IsInitialized() const override { return bInitialized; }
 
 	virtual const FEOSSDKPlatformConfig* GetPlatformConfig(const FString& PlatformConfigName, bool bLoadIfMissing = false) override;
-	virtual bool AddPlatformConfig(const FEOSSDKPlatformConfig& PlatformConfig) override;
+	virtual bool AddPlatformConfig(const FEOSSDKPlatformConfig& PlatformConfig, bool bOverwriteExistingConfig = false) override;
 	virtual const FString& GetDefaultPlatformConfigName() override;
 	virtual void SetDefaultPlatformConfigName(const FString& PlatformConfigName) override;
 
diff --git a/Engine/Plugins/Online/EOSShared/Source/EOSShared/Public/IEOSSDKManager.h b/Engine/Plugins/Online/EOSShared/Source/EOSShared/Public/IEOSSDKManager.h
--- a/Engine/Plugins/Online/EOSShared/Source/EOSShared/Public/IEOSSDKManager.h
+++ b/Engine/Plugins/Online/EOSShared/Source/EOSShared/Public/IEOSSDKManager.h
@@ -104,7 +104,7 @@ public:
 	virtual bool IsInitialized() const = 0;
 
 	virtual const FEOSSDKPlatformConfig* GetPlatformConfig(const FString& PlatformConfigName, bool bLoadIfMissing = false) = 0;
-	virtual bool AddPlatformConfig(const FEOSSDKPlatformConfig& PlatformConfig) = 0;
+	virtual bool AddPlatformConfig(const FEOSSDKPlatformConfig& PlatformConfig, bool bOverwriteExistingConfig = false) = 0;
 	virtual const FString& GetDefaultPlatformConfigName() = 0;
 	virtual void SetDefaultPlatformConfigName(const FString& PlatformConfigName) = 0;
 
