Fix for Compile UnrealGame LinuxArm64 error due to missing includes when !WITH_EOS

diff --git a/Engine/Plugins/Online/EOSShared/Source/EOSShared/Private/EOSSharedModule.h b/Engine/Plugins/Online/EOSShared/Source/EOSShared/Private/EOSSharedModule.h
--- a/Engine/Plugins/Online/EOSShared/Source/EOSShared/Private/EOSSharedModule.h
+++ b/Engine/Plugins/Online/EOSShared/Source/EOSShared/Private/EOSSharedModule.h
@@ -2,6 +2,9 @@
 
 #pragma once
 
+#include "Containers/Array.h"
+#include "Containers/UnrealString.h"
+#include "Delegates/IDelegateInstance.h"
 #include "Misc/Fork.h"
 #include "Modules/ModuleInterface.h"
 
