- Add "FastLock" RAII mechanism that locks the SDKManager to ticking the SDK platforms every frame, ignoring the config interval.

diff --git a/Engine/Plugins/Online/EOSShared/Source/EOSShared/Private/EOSSDKManager.cpp b/Engine/Plugins/Online/EOSShared/Source/EOSShared/Private/EOSSDKManager.cpp
--- a/Engine/Plugins/Online/EOSShared/Source/EOSShared/Private/EOSSDKManager.cpp
+++ b/Engine/Plugins/Online/EOSShared/Source/EOSShared/Private/EOSSDKManager.cpp
@@ -753,7 +753,7 @@ void FEOSSDKManager::LoadConfig()
 	SetupTicker();
 }
 
-void FEOSSDKManager::SetupTicker(bool bIgnoreConfigTickInterval)
+void FEOSSDKManager::SetupTicker()
 {
 	check(IsInGameThread());
 
@@ -766,7 +766,8 @@ void FEOSSDKManager::SetupTicker(bool bIgnoreConfigTickInterval)
 	int NumActivePlatforms = ActivePlatforms.Num();
 	if (NumActivePlatforms > 0)
 	{
-		const double TickIntervalSeconds = ConfigTickIntervalSeconds > SMALL_NUMBER && !bIgnoreConfigTickInterval ? ConfigTickIntervalSeconds / NumActivePlatforms : 0.f;
+		const bool bIsFastTicking = FastTickLock.IsValid();
+		const double TickIntervalSeconds = ConfigTickIntervalSeconds > SMALL_NUMBER && !bIsFastTicking ? ConfigTickIntervalSeconds / NumActivePlatforms : 0.f;
 		TickerHandle = FTSTicker::GetCoreTicker().AddTicker(FTickerDelegate::CreateRaw(this, &FEOSSDKManager::Tick), TickIntervalSeconds);
 	}
 }
@@ -859,7 +860,15 @@ void FEOSSDKManager::SetInvokeOverlayButton(const EOS_HPlatform PlatformHandle)
 
 void FEOSSDKManager::OnDisplaySettingsUpdated(const EOS_UI_OnDisplaySettingsUpdatedCallbackInfo* Data)
 {
-	reinterpret_cast<FEOSSDKManager*>(Data->ClientData)->SetupTicker((bool)Data->bIsVisible);
+	static TSharedPtr<IEOSFastTickLock> OverlayFastTickLock;
+	if (Data->bIsVisible && !OverlayFastTickLock.IsValid())
+	{
+		OverlayFastTickLock = reinterpret_cast<FEOSSDKManager*>(Data->ClientData)->GetFastTickLock();
+	}
+	else if (!Data->bIsVisible && OverlayFastTickLock.IsValid())
+	{
+		OverlayFastTickLock.Reset();
+	}
 }
 
 void FEOSSDKManager::RegisterDisplaySettingsUpdatedCallback(const EOS_HPlatform PlatformHandle)
@@ -907,7 +916,7 @@ bool FEOSSDKManager::Tick(float)
 		if (ConfigTickIntervalSeconds > SMALL_NUMBER)
 		{
 			PlatformTickIdx = (PlatformTickIdx + 1) % ActivePlatformHandles.Num();
-			PlatformsToTick.Emplace(ActivePlatformHandles[PlatformTickIdx]);	
+			PlatformsToTick.Emplace(ActivePlatformHandles[PlatformTickIdx]);
 		}
 		else
 		{
@@ -1525,6 +1534,21 @@ void FEOSSDKManager::AddCallbackObject(TUniquePtr<FCallbackBase> CallbackObj)
 	CallbackObjects.Emplace(MoveTemp(CallbackObj));
 }
 
+TSharedRef<IEOSFastTickLock> FEOSSDKManager::GetFastTickLock()
+{
+	if (TSharedPtr<IEOSFastTickLock> ExistingLock = FastTickLock.Pin())
+	{
+		return ExistingLock.ToSharedRef();
+	}
+
+	TSharedRef<IEOSFastTickLock> NewLock = MakeShared<FEOSFastTickLock>();
+	FastTickLock = NewLock;
+
+	SetupTicker();
+
+	return NewLock;
+}
+
 FEOSPlatformHandle::~FEOSPlatformHandle()
 {
 	Manager.ReleasePlatform(PlatformHandle);
@@ -1538,6 +1562,11 @@ void FEOSPlatformHandle::Tick()
 	EOS_Platform_Tick(PlatformHandle);
 }
 
+TSharedRef<IEOSFastTickLock> FEOSPlatformHandle::GetFastTickLock()
+{
+	return Manager.GetFastTickLock();
+}
+
 FString FEOSPlatformHandle::GetConfigName() const
 {
 	return ConfigName;
@@ -1583,4 +1612,13 @@ void FEOSPlatformHandle::LogConnectInfo(const EOS_ProductUserId LoggedInAccount,
 	Manager.LogConnectInfo(PlatformHandle, LoggedInAccount, Indent);
 }
 
+FEOSFastTickLock::~FEOSFastTickLock()
+{
+	if (FEOSSDKManager* Manager = static_cast<FEOSSDKManager*>(IEOSSDKManager::Get()))
+	{
+		Manager->FastTickLock.Reset();
+		Manager->SetupTicker();
+	}
+}
+
 #endif // WITH_EOS_SDK
\ No newline at end of file
diff --git a/Engine/Plugins/Online/EOSShared/Source/EOSShared/Private/EOSSDKManager.h b/Engine/Plugins/Online/EOSShared/Source/EOSShared/Private/EOSSDKManager.h
--- a/Engine/Plugins/Online/EOSShared/Source/EOSShared/Private/EOSSDKManager.h
+++ b/Engine/Plugins/Online/EOSShared/Source/EOSShared/Private/EOSSDKManager.h
@@ -63,6 +63,8 @@ public:
 	virtual void LogConnectInfo(const EOS_HPlatform Platform, const EOS_ProductUserId LoggedInAccount, int32 Indent = 0) const override;
 
 	virtual void AddCallbackObject(TUniquePtr<class FCallbackBase> CallbackObj) override;
+
+	virtual TSharedRef<IEOSFastTickLock> GetFastTickLock() override;
 	// End IEOSSDKManager
 
 	EOS_EResult Initialize();
@@ -105,12 +107,13 @@ protected:
 	EOS_EApplicationStatus CachedApplicationStatus = EOS_EApplicationStatus::EOS_AS_Foreground;
 
 	friend struct FEOSPlatformHandle;
+	friend struct FEOSFastTickLock;
 
 	void OnConfigSectionsChanged(const FString& IniFilename, const TSet<FString>& SectionName);
 	void LoadConfig();
 	void ReleasePlatform(EOS_HPlatform PlatformHandle);
 	void ReleaseReleasedPlatforms();
-	void SetupTicker(bool bIgnoreConfigTickInterval = false);
+	void SetupTicker();
 	void OnLogVerbosityChanged(const FLogCategoryName& CategoryName, ELogVerbosity::Type OldVerbosity, ELogVerbosity::Type NewVerbosity);
 
 #if EOSSDK_RUNTIME_LOAD_REQUIRED
@@ -155,6 +158,9 @@ protected:
 
 	/** Management flags passed on as options in integrated platform setup */
 	EOS_EIntegratedPlatformManagementFlags IntegratedPlatformManagementFlags = {};
+
+	/** If this is set, then we should ignore the config tick rate and tick at full speed. */
+	TWeakPtr<IEOSFastTickLock> FastTickLock;
 };
 
 struct FEOSPlatformHandle : public IEOSPlatformHandle
@@ -166,6 +172,7 @@ struct FEOSPlatformHandle : public IEOSPlatformHandle
 
 	virtual ~FEOSPlatformHandle();
 	virtual void Tick() override;
+	virtual TSharedRef<IEOSFastTickLock> GetFastTickLock() override;
 
 	virtual FString GetConfigName() const override;
 	virtual FString GetOverrideCountryCode() const override;
@@ -185,4 +192,9 @@ struct FEOSPlatformHandle : public IEOSPlatformHandle
 	FString ConfigName;
 };
 
+struct FEOSFastTickLock : public IEOSFastTickLock
+{
+	virtual ~FEOSFastTickLock();
+};
+
 #endif // WITH_EOS_SDK
\ No newline at end of file
diff --git a/Engine/Plugins/Online/EOSShared/Source/EOSShared/Public/IEOSSDKManager.h b/Engine/Plugins/Online/EOSShared/Source/EOSShared/Public/IEOSSDKManager.h
--- a/Engine/Plugins/Online/EOSShared/Source/EOSShared/Public/IEOSSDKManager.h
+++ b/Engine/Plugins/Online/EOSShared/Source/EOSShared/Public/IEOSSDKManager.h
@@ -39,11 +39,15 @@ struct FEOSSDKPlatformConfig
 	TArray<FString> OptionalConfig;
 };
 
-DECLARE_MULTICAST_DELEGATE_OneParam(FEOSSDKManagerOnPreInitializeSDK, EOS_InitializeOptions& Options);
-DECLARE_MULTICAST_DELEGATE_OneParam(FEOSSDKManagerOnPostInitializeSDK, EOS_EResult Result);
-DECLARE_MULTICAST_DELEGATE_TwoParams(FEOSSDKManagerOnDefaultPlatformConfigNameChanged, const FString& NewName, const FString& OldName);
-DECLARE_MULTICAST_DELEGATE_TwoParams(FEOSSDKManagerOnPreCreateNamedPlatform, const FEOSSDKPlatformConfig& Config, EOS_Platform_Options& Options);
-DECLARE_MULTICAST_DELEGATE_OneParam(FEOSSDKManagerOnPreCreatePlatform, EOS_Platform_Options& Options);
+/**
+* Allows temporary RAII style overriding of the config driven tick rate, for scenarios where you want the SDK
+* To tick as fast as possible, i.e. overlay visible, or time critical operations
+*/
+class IEOSFastTickLock
+{
+public:
+	virtual ~IEOSFastTickLock() = default;
+};
 
 class IEOSPlatformHandle
 {
@@ -55,6 +59,7 @@ public:
 	virtual ~IEOSPlatformHandle() = default;
 
 	virtual void Tick() = 0;
+	virtual TSharedRef<IEOSFastTickLock> GetFastTickLock() = 0;
 
 	operator EOS_HPlatform() const { return PlatformHandle; }
 
@@ -76,6 +81,11 @@ protected:
 using IEOSPlatformHandlePtr = TSharedPtr<IEOSPlatformHandle, ESPMode::ThreadSafe>;
 using IEOSPlatformHandleWeakPtr = TWeakPtr<IEOSPlatformHandle, ESPMode::ThreadSafe>;
 
+DECLARE_MULTICAST_DELEGATE_OneParam(FEOSSDKManagerOnPreInitializeSDK, EOS_InitializeOptions& Options);
+DECLARE_MULTICAST_DELEGATE_OneParam(FEOSSDKManagerOnPostInitializeSDK, EOS_EResult Result);
+DECLARE_MULTICAST_DELEGATE_TwoParams(FEOSSDKManagerOnDefaultPlatformConfigNameChanged, const FString& NewName, const FString& OldName);
+DECLARE_MULTICAST_DELEGATE_TwoParams(FEOSSDKManagerOnPreCreateNamedPlatform, const FEOSSDKPlatformConfig& Config, EOS_Platform_Options& Options);
+DECLARE_MULTICAST_DELEGATE_OneParam(FEOSSDKManagerOnPreCreatePlatform, EOS_Platform_Options& Options);
 DECLARE_MULTICAST_DELEGATE_OneParam(FEOSSDKManagerOnPlatformCreated, const IEOSPlatformHandlePtr& PlatformHandle);
 DECLARE_MULTICAST_DELEGATE_OneParam(FEOSSDKManagerOnPreReleasePlatform, const EOS_HPlatform& PlatformHandle);
 
@@ -135,6 +145,8 @@ public:
 	/** Assign ownership of a callback object, to be free'd after EOS_Shutdown */
 	virtual void AddCallbackObject(TUniquePtr<class FCallbackBase> CallbackObj) = 0;
 
+	virtual TSharedRef<IEOSFastTickLock> GetFastTickLock() = 0;
+
 	FEOSSDKManagerOnPreInitializeSDK OnPreInitializeSDK;
 	FEOSSDKManagerOnPostInitializeSDK OnPostInitializeSDK;
 	FEOSSDKManagerOnDefaultPlatformConfigNameChanged OnDefaultPlatformConfigNameChanged;
