- Add new OnPreInitializeSDK callback variant that allows swapping the entire Options struct out. This ensures bindees can support ApiVersions greater than the version that Unreal was shipped against, with new fields, by creating a new options struct from the new headers, and swapping it.

diff --git a/Engine/Plugins/Online/EOSShared/Source/EOSShared/Private/EOSSDKManager.cpp b/Engine/Plugins/Online/EOSShared/Source/EOSShared/Private/EOSSDKManager.cpp
--- a/Engine/Plugins/Online/EOSShared/Source/EOSShared/Private/EOSSDKManager.cpp
+++ b/Engine/Plugins/Online/EOSShared/Source/EOSShared/Private/EOSSDKManager.cpp
@@ -1145,7 +1145,10 @@ EOS_EResult FEOSSDKManager::EOSInitialize(EOS_InitializeOptions& Options)
 {
 	OnPreInitializeSDK.Broadcast(Options);
 
-	return EOS_Initialize(&Options);
+	EOS_InitializeOptions* OptionsPtr = &Options;
+	OnPreInitializeSDK2.Broadcast(OptionsPtr);
+
+	return EOS_Initialize(OptionsPtr);
 }
 
 bool FEOSSDKManager::Exec_Runtime(class UWorld* InWorld, const TCHAR* Cmd, FOutputDevice& Ar)
diff --git a/Engine/Plugins/Online/EOSShared/Source/EOSShared/Public/IEOSSDKManager.h b/Engine/Plugins/Online/EOSShared/Source/EOSShared/Public/IEOSSDKManager.h
--- a/Engine/Plugins/Online/EOSShared/Source/EOSShared/Public/IEOSSDKManager.h
+++ b/Engine/Plugins/Online/EOSShared/Source/EOSShared/Public/IEOSSDKManager.h
@@ -81,7 +81,10 @@ protected:
 using IEOSPlatformHandlePtr = TSharedPtr<IEOSPlatformHandle, ESPMode::ThreadSafe>;
 using IEOSPlatformHandleWeakPtr = TWeakPtr<IEOSPlatformHandle, ESPMode::ThreadSafe>;
 
+// This callback lets you modify the options struct
 DECLARE_MULTICAST_DELEGATE_OneParam(FEOSSDKManagerOnPreInitializeSDK, EOS_InitializeOptions& Options);
+// This callback lets you modify or replace the options struct
+DECLARE_MULTICAST_DELEGATE_OneParam(FEOSSDKManagerOnPreInitializeSDK2, EOS_InitializeOptions*& InOutOptions);
 DECLARE_MULTICAST_DELEGATE_OneParam(FEOSSDKManagerOnPostInitializeSDK, EOS_EResult Result);
 DECLARE_MULTICAST_DELEGATE_TwoParams(FEOSSDKManagerOnDefaultPlatformConfigNameChanged, const FString& NewName, const FString& OldName);
 DECLARE_MULTICAST_DELEGATE_TwoParams(FEOSSDKManagerOnPreCreateNamedPlatform, const FEOSSDKPlatformConfig& Config, EOS_Platform_Options& Options);
@@ -148,6 +151,7 @@ public:
 	virtual TSharedRef<IEOSFastTickLock> GetFastTickLock() = 0;
 
 	FEOSSDKManagerOnPreInitializeSDK OnPreInitializeSDK;
+	FEOSSDKManagerOnPreInitializeSDK2 OnPreInitializeSDK2;
 	FEOSSDKManagerOnPostInitializeSDK OnPostInitializeSDK;
 	FEOSSDKManagerOnDefaultPlatformConfigNameChanged OnDefaultPlatformConfigNameChanged;
 	FEOSSDKManagerOnPreCreateNamedPlatform OnPreCreateNamedPlatform;
