- Fix ProductName being empty in UnrealEditor-Cmd commandlet runs, causing EOS_Initialize to fail

#rnx

diff --git a/Engine/Plugins/Online/EOSShared/Source/EOSShared/Private/EOSSDKManager.cpp b/Engine/Plugins/Online/EOSShared/Source/EOSShared/Private/EOSSDKManager.cpp
--- a/Engine/Plugins/Online/EOSShared/Source/EOSShared/Private/EOSSDKManager.cpp
+++ b/Engine/Plugins/Online/EOSShared/Source/EOSShared/Private/EOSSDKManager.cpp
@@ -995,10 +995,18 @@ void FEOSSDKManager::OnLogVerbosityChanged(const FLogCategoryName& CategoryName,
 FString FEOSSDKManager::GetProductName() const
 {
 	FString ProductName;
-	if (!GConfig->GetString(TEXT("EOSSDK"), TEXT("ProductName"), ProductName, GEngineIni))
+	GConfig->GetString(TEXT("EOSSDK"), TEXT("ProductName"), ProductName, GEngineIni);
+
+	if (ProductName.IsEmpty())
 	{
 		ProductName = FApp::GetProjectName();
 	}
+
+	if (ProductName.IsEmpty())
+	{
+		ProductName = TEXT("UnrealEngine");
+	}
+
 	return ProductName;
 }
 
