- Make EOSShared return immediately from StartupModule if something else has already registered the EOSSDKManager modular feature. This allows licensees to provide their own EOSSDKManager implementation.

#rnx

diff --git a/Engine/Plugins/Online/EOSShared/Source/EOSShared/Private/EOSSharedModule.cpp b/Engine/Plugins/Online/EOSShared/Source/EOSShared/Private/EOSSharedModule.cpp
--- a/Engine/Plugins/Online/EOSShared/Source/EOSShared/Private/EOSSharedModule.cpp
+++ b/Engine/Plugins/Online/EOSShared/Source/EOSShared/Private/EOSSharedModule.cpp
@@ -21,7 +21,9 @@ IMPLEMENT_MODULE(FEOSSharedModule, EOSShared);
 void FEOSSharedModule::StartupModule()
 {
 #if WITH_EOS_SDK
-	if (FParse::Param(FCommandLine::Get(), TEXT("NoEOS")))
+	const bool bNoEOS = FParse::Param(FCommandLine::Get(), TEXT("NoEOS"));
+	const bool bAlreadyRegistered = IModularFeatures::Get().IsModularFeatureAvailable(IEOSSDKManager::GetModularFeatureName());
+	if (bNoEOS || bAlreadyRegistered)
 	{
 		return;
 	}
