- Move Core and EOSSDK to PublicDependencyModuleNames in EOSShared.Build.cs, as the public headers depend on both.
- Add delegate for getting the runtime library name to load.
- Remove SDK preload mechanism. This dates back to before EOSShared existed, to a time when SDK initialization was much later in startup flow, but now this runs _after_ EOSShared StartupModule and serves no purpose

#rnx

diff --git a/Engine/Plugins/Online/EOSShared/Source/EOSShared/EOSShared.Build.cs b/Engine/Plugins/Online/EOSShared/Source/EOSShared/EOSShared.Build.cs
--- a/Engine/Plugins/Online/EOSShared/Source/EOSShared/EOSShared.Build.cs
+++ b/Engine/Plugins/Online/EOSShared/Source/EOSShared/EOSShared.Build.cs
@@ -25,7 +25,7 @@ public class EOSShared : ModuleRules
 
 		PublicDefinitions.Add("UE_WITH_EOS_SDK_APIVERSION_WARNINGS=" + (EnableApiVersionWarnings ? "1" : "0"));
 
-		PrivateDependencyModuleNames.AddRange(
+		PublicDependencyModuleNames.AddRange(
 			new string[]
 			{
 				"Core",
diff --git a/Engine/Plugins/Online/EOSShared/Source/EOSShared/Private/EOSSDKManager.cpp b/Engine/Plugins/Online/EOSShared/Source/EOSShared/Private/EOSSDKManager.cpp
--- a/Engine/Plugins/Online/EOSShared/Source/EOSShared/Private/EOSSDKManager.cpp
+++ b/Engine/Plugins/Online/EOSShared/Source/EOSShared/Private/EOSSDKManager.cpp
@@ -150,73 +150,52 @@ namespace
 		}
 	}
 #endif // !NO_LOGGING
+}
 
 #if EOSSDK_RUNTIME_LOAD_REQUIRED
-	static void* GetSdkDllHandle()
-	{
-		void* Result = nullptr;
+void FEOSSDKManager::LoadSDKHandle()
+{
+	check(SDKHandle == nullptr);
 
-		const FString ProjectBinaryPath = FPaths::ConvertRelativePathToFull(FPaths::Combine(FPaths::ProjectDir(), TEXT("Binaries"), FPlatformProcess::GetBinariesSubdirectory(), TEXT(EOSSDK_RUNTIME_LIBRARY_NAME)));
-		if (FPaths::FileExists(ProjectBinaryPath))
-		{
-			Result = FPlatformProcess::GetDllHandle(*ProjectBinaryPath);
-		}
+	const FString RuntimeLibraryName = OnRequestRuntimeLibraryName.IsBound()
+		? OnRequestRuntimeLibraryName.Execute()
+		: TEXT(EOSSDK_RUNTIME_LIBRARY_NAME);
 
-		if (!Result)
-		{
-			const FString EngineBinaryPath = FPaths::ConvertRelativePathToFull(FPaths::Combine(FPaths::EngineDir(), TEXT("Binaries"), FPlatformProcess::GetBinariesSubdirectory(), TEXT(EOSSDK_RUNTIME_LIBRARY_NAME)));
-			if (FPaths::FileExists(EngineBinaryPath))
-			{
-				Result = FPlatformProcess::GetDllHandle(*EngineBinaryPath);
-			}
-		}
+	const FString ProjectBinaryPath = FPaths::ConvertRelativePathToFull(FPaths::Combine(FPaths::ProjectDir(), TEXT("Binaries"), FPlatformProcess::GetBinariesSubdirectory(), *RuntimeLibraryName));
+	if (FPaths::FileExists(ProjectBinaryPath))
+	{
+		SDKHandle = FPlatformProcess::GetDllHandle(*ProjectBinaryPath);
+	}
 
-		if (!Result)
+	if (!SDKHandle)
+	{
+		const FString EngineBinaryPath = FPaths::ConvertRelativePathToFull(FPaths::Combine(FPaths::EngineDir(), TEXT("Binaries"), FPlatformProcess::GetBinariesSubdirectory(), *RuntimeLibraryName));
+		if (FPaths::FileExists(EngineBinaryPath))
 		{
-			Result = FPlatformProcess::GetDllHandle(TEXT(EOSSDK_RUNTIME_LIBRARY_NAME));
+			SDKHandle = FPlatformProcess::GetDllHandle(*EngineBinaryPath);
 		}
+	}
 
-		if (!Result)
+	if (!SDKHandle)
+	{
+		SDKHandle = FPlatformProcess::GetDllHandle(*RuntimeLibraryName);
+	}
+
+	if (!SDKHandle)
+	{
+		bool bDllLoadFailureIsFatal = false;
+		GConfig->GetBool(TEXT("EOSSDK"), TEXT("bDllLoadFailureIsFatal"), bDllLoadFailureIsFatal, GEngineIni);
+		if (bDllLoadFailureIsFatal)
 		{
-			bool bDllLoadFailureIsFatal = false;
-			GConfig->GetBool(TEXT("EOSSDK"), TEXT("bDllLoadFailureIsFatal"), bDllLoadFailureIsFatal, GEngineIni);
-			if (bDllLoadFailureIsFatal)
-			{
-				FPlatformMisc::MessageBoxExt(
-					EAppMsgType::Ok,
-					*FText::Format(NSLOCTEXT("EOSShared", "DllLoadFail", "Failed to load {0}. Please verify your installation. Exiting..."), FText::FromString(TEXT(EOSSDK_RUNTIME_LIBRARY_NAME))).ToString(),
-					TEXT("Error")
-				);
-				UE_LOG(LogEOSSDK, Fatal, TEXT("Failed to load EOSSDK binary"));
-			}
+			FPlatformMisc::MessageBoxExt(
+				EAppMsgType::Ok,
+				*FText::Format(NSLOCTEXT("EOSShared", "DllLoadFail", "Failed to load {0}. Please verify your installation. Exiting..."), FText::FromString(*RuntimeLibraryName)).ToString(),
+				TEXT("Error")
+			);
+			UE_LOG(LogEOSSDK, Fatal, TEXT("Failed to load EOSSDK binary"));
 		}
-
-		return Result;
 	}
-#endif
 }
-
-
-// if the platform wants to preload the EOS runtime library, then use the auto registration system to kick off
-// the GetDllHandle call at a specified time. it is assumed that calling GetDllHandle here will not conflict
-// with the call in FEOSSDKManager::FEOSSDKManager, and will only speed it up, even if it's not complete before
-// the later one is called. if that's not the case, we will need an FEvent or similar to block the constructor
-// until this is complete
-#if EOSSDK_RUNTIME_LOAD_REQUIRED && defined(EOS_DLL_PRELOAD_PHASE)
-
-#include "Misc/DelayedAutoRegister.h"
-#include "Async/Async.h"
-
-static FDelayedAutoRegisterHelper GKickoffDll(EDelayedRegisterRunPhase::EOS_DLL_PRELOAD_PHASE, []
-{
-	Async(EAsyncExecution::Thread, []
-	{
-		LLM_SCOPE(ELLMTag::RealTimeCommunications);
-		SCOPED_BOOT_TIMING("Preloading EOS module");
-		GetSdkDllHandle();
-	});
-});
-
 #endif
 
 FEOSSDKManager::FEOSSDKManager()
@@ -245,7 +224,7 @@ EOS_EResult FEOSSDKManager::Initialize()
 	if (SDKHandle == nullptr)
 	{
 		LLM_SCOPE(ELLMTag::RealTimeCommunications);
-		SDKHandle = GetSdkDllHandle();
+		LoadSDKHandle();
 	}
 	if (SDKHandle == nullptr)
 	{
diff --git a/Engine/Plugins/Online/EOSShared/Source/EOSShared/Private/EOSSDKManager.h b/Engine/Plugins/Online/EOSShared/Source/EOSShared/Private/EOSSDKManager.h
--- a/Engine/Plugins/Online/EOSShared/Source/EOSShared/Private/EOSSDKManager.h
+++ b/Engine/Plugins/Online/EOSShared/Source/EOSShared/Private/EOSSDKManager.h
@@ -117,6 +117,7 @@ protected:
 	void OnLogVerbosityChanged(const FLogCategoryName& CategoryName, ELogVerbosity::Type OldVerbosity, ELogVerbosity::Type NewVerbosity);
 
 #if EOSSDK_RUNTIME_LOAD_REQUIRED
+	void LoadSDKHandle();
 	void* SDKHandle = nullptr;
 #endif
 
diff --git a/Engine/Plugins/Online/EOSShared/Source/EOSShared/Public/IEOSSDKManager.h b/Engine/Plugins/Online/EOSShared/Source/EOSShared/Public/IEOSSDKManager.h
--- a/Engine/Plugins/Online/EOSShared/Source/EOSShared/Public/IEOSSDKManager.h
+++ b/Engine/Plugins/Online/EOSShared/Source/EOSShared/Public/IEOSSDKManager.h
@@ -91,6 +91,7 @@ DECLARE_MULTICAST_DELEGATE_TwoParams(FEOSSDKManagerOnPreCreateNamedPlatform, con
 DECLARE_MULTICAST_DELEGATE_OneParam(FEOSSDKManagerOnPreCreatePlatform, EOS_Platform_Options& Options);
 DECLARE_MULTICAST_DELEGATE_OneParam(FEOSSDKManagerOnPlatformCreated, const IEOSPlatformHandlePtr& PlatformHandle);
 DECLARE_MULTICAST_DELEGATE_OneParam(FEOSSDKManagerOnPreReleasePlatform, const EOS_HPlatform& PlatformHandle);
+DECLARE_DELEGATE_RetVal(FString, FEOSSDKManagerOnRequestRuntimeLibraryName);
 
 class IEOSSDKManager : public IModularFeature
 {
@@ -158,6 +159,7 @@ public:
 	FEOSSDKManagerOnPreCreatePlatform OnPreCreatePlatform;
 	FEOSSDKManagerOnPlatformCreated OnPlatformCreated;
 	FEOSSDKManagerOnPreReleasePlatform OnPreReleasePlatform;
+	FEOSSDKManagerOnRequestRuntimeLibraryName OnRequestRuntimeLibraryName;
 };
 
 #endif // WITH_EOS_SDK
\ No newline at end of file
