Add LLM tags in EOSVirtualMemoryCallbacks

diff --git a/Engine/Plugins/Online/EOSShared/Source/EOSShared/Private/EOSVirtualMemoryCallbacks.cpp b/Engine/Plugins/Online/EOSShared/Source/EOSShared/Private/EOSVirtualMemoryCallbacks.cpp
--- a/Engine/Plugins/Online/EOSShared/Source/EOSShared/Private/EOSVirtualMemoryCallbacks.cpp
+++ b/Engine/Plugins/Online/EOSShared/Source/EOSShared/Private/EOSVirtualMemoryCallbacks.cpp
@@ -4,6 +4,8 @@
 
 #if WITH_EOS_SDK
 
+LLM_DEFINE_TAG(KITT);
+
 #include "HAL/PlatformMemory.h"
 
 namespace EOSVirtualMemoryCallbacks
@@ -11,6 +13,7 @@ namespace EOSVirtualMemoryCallbacks
 
 void* EOS_MEMORY_CALL Reserve(size_t SizeInBytes, EOS_EVM_AccessType AccessType, void** OutContextData)
 {
+	LLM_PLATFORM_SCOPE_BYTAG(KITT);
 	FPlatformMemory::FPlatformVirtualMemoryBlock* Block = new FPlatformMemory::FPlatformVirtualMemoryBlock(FPlatformMemory::FPlatformVirtualMemoryBlock::AllocateVirtual(SizeInBytes));
 	*OutContextData = Block;
 	return Block->GetVirtualPointer();
@@ -18,6 +21,7 @@ void* EOS_MEMORY_CALL Reserve(size_t SizeInBytes, EOS_EVM_AccessType AccessType,
 
 EOS_Bool EOS_MEMORY_CALL Release(void* Address, size_t SizeInBytes, void* ContextData)
 {
+	LLM_PLATFORM_SCOPE_BYTAG(KITT);
 	FPlatformMemory::FPlatformVirtualMemoryBlock* Block = static_cast<FPlatformMemory::FPlatformVirtualMemoryBlock*>(ContextData);
 	Block->FreeVirtual();
 	delete Block;
@@ -26,15 +30,19 @@ EOS_Bool EOS_MEMORY_CALL Release(void* Address, size_t SizeInBytes, void* Contex
 
 EOS_Bool EOS_MEMORY_CALL Commit(void* Address, size_t SizeInBytes, void* ContextData)
 {
+	LLM_PLATFORM_SCOPE_BYTAG(KITT);
 	FPlatformMemory::FPlatformVirtualMemoryBlock* Block = static_cast<FPlatformMemory::FPlatformVirtualMemoryBlock*>(ContextData);
 	Block->CommitByPtr(Address, SizeInBytes);
+	LLM_IF_ENABLED(FLowLevelMemTracker::Get().OnLowLevelChangeInMemoryUse(ELLMTracker::Platform, static_cast<int64>(SizeInBytes)));
 	return true;
 }
 
 EOS_Bool EOS_MEMORY_CALL Decommit(void* Address, size_t SizeInBytes, void* ContextData)
 {
+	LLM_PLATFORM_SCOPE_BYTAG(KITT);
 	FPlatformMemory::FPlatformVirtualMemoryBlock* Block = static_cast<FPlatformMemory::FPlatformVirtualMemoryBlock*>(ContextData);
 	Block->DecommitByPtr(Address, SizeInBytes);
+	LLM_IF_ENABLED(FLowLevelMemTracker::Get().OnLowLevelChangeInMemoryUse(ELLMTracker::Platform, -static_cast<int64>(SizeInBytes)));
 	return true;
 }
 
diff --git a/Engine/Plugins/Online/EOSShared/Source/EOSShared/Private/EOSVirtualMemoryCallbacks.h b/Engine/Plugins/Online/EOSShared/Source/EOSShared/Private/EOSVirtualMemoryCallbacks.h
--- a/Engine/Plugins/Online/EOSShared/Source/EOSShared/Private/EOSVirtualMemoryCallbacks.h
+++ b/Engine/Plugins/Online/EOSShared/Source/EOSShared/Private/EOSVirtualMemoryCallbacks.h
@@ -4,6 +4,8 @@
 
 #if WITH_EOS_SDK
 
+#include "HAL/LowLevelMemTracker.h"
+
 #if defined(EOS_PLATFORM_BASE_FILE_NAME)
 #include EOS_PLATFORM_BASE_FILE_NAME
 #endif
@@ -11,6 +13,8 @@
 
 enum class EOS_EVM_AccessType : int32_t;
 
+LLM_DECLARE_TAG_API(KITT, EOSSHARED_API);
+
 // Generic VMem callback implementations used by most platforms requiring VMem callbacks
 namespace EOSVirtualMemoryCallbacks
 {
