- Add ability to specify the EOSSDK binary to use on command line

#tests ran with -eossdkbinary= pointing at an old SDK DLL, observed that it loaded and failed to init because the options struct was too new.

diff --git a/Engine/Plugins/Online/EOSShared/Source/EOSShared/Private/EOSSDKManager.cpp b/Engine/Plugins/Online/EOSShared/Source/EOSShared/Private/EOSSDKManager.cpp
--- a/Engine/Plugins/Online/EOSShared/Source/EOSShared/Private/EOSSDKManager.cpp
+++ b/Engine/Plugins/Online/EOSShared/Source/EOSShared/Private/EOSSDKManager.cpp
@@ -161,24 +161,36 @@ void FEOSSDKManager::LoadSDKHandle()
 		? OnRequestRuntimeLibraryName.Execute()
 		: TEXT(EOSSDK_RUNTIME_LIBRARY_NAME);
 
-	const FString ProjectBinaryPath = FPaths::ConvertRelativePathToFull(FPaths::Combine(FPaths::ProjectDir(), TEXT("Binaries"), FPlatformProcess::GetBinariesSubdirectory(), *RuntimeLibraryName));
-	if (FPaths::FileExists(ProjectBinaryPath))
+#if !UE_BUILD_SHIPPING
+	FString RuntimeLibraryOverride;
+	if (FParse::Value(FCommandLine::Get(), TEXT("eossdkbinary="), RuntimeLibraryOverride))
 	{
-		SDKHandle = FPlatformProcess::GetDllHandle(*ProjectBinaryPath);
+		SDKHandle = FPlatformProcess::GetDllHandle(*RuntimeLibraryOverride);
+		UE_CLOG(SDKHandle, LogEOSShared, Log, TEXT("%hs Loaded eos binary \"%s\""), __FUNCTION__, *RuntimeLibraryOverride);
+		UE_CLOG(!SDKHandle, LogEOSShared, Error, TEXT("%hs Failed to load eos binary \"%s\""), __FUNCTION__, *RuntimeLibraryOverride);
 	}
-
-	if (!SDKHandle)
+	else
+#endif // !UE_BUILD_SHIPPING
 	{
-		const FString EngineBinaryPath = FPaths::ConvertRelativePathToFull(FPaths::Combine(FPaths::EngineDir(), TEXT("Binaries"), FPlatformProcess::GetBinariesSubdirectory(), *RuntimeLibraryName));
-		if (FPaths::FileExists(EngineBinaryPath))
+		const FString ProjectBinaryPath = FPaths::ConvertRelativePathToFull(FPaths::Combine(FPaths::ProjectDir(), TEXT("Binaries"), FPlatformProcess::GetBinariesSubdirectory(), *RuntimeLibraryName));
+		if (FPaths::FileExists(ProjectBinaryPath))
 		{
-			SDKHandle = FPlatformProcess::GetDllHandle(*EngineBinaryPath);
+			SDKHandle = FPlatformProcess::GetDllHandle(*ProjectBinaryPath);
 		}
-	}
 
-	if (!SDKHandle)
-	{
-		SDKHandle = FPlatformProcess::GetDllHandle(*RuntimeLibraryName);
+		if (!SDKHandle)
+		{
+			const FString EngineBinaryPath = FPaths::ConvertRelativePathToFull(FPaths::Combine(FPaths::EngineDir(), TEXT("Binaries"), FPlatformProcess::GetBinariesSubdirectory(), *RuntimeLibraryName));
+			if (FPaths::FileExists(EngineBinaryPath))
+			{
+				SDKHandle = FPlatformProcess::GetDllHandle(*EngineBinaryPath);
+			}
+		}
+
+		if (!SDKHandle)
+		{
+			SDKHandle = FPlatformProcess::GetDllHandle(*RuntimeLibraryName);
+		}
 	}
 
 	if (!SDKHandle)
