Add SteamShared dependency to EOSShared and Lyra SteamEOS config

diff --git a/Engine/Plugins/Online/EOSShared/EOSShared.uplugin b/Engine/Plugins/Online/EOSShared/EOSShared.uplugin
--- a/Engine/Plugins/Online/EOSShared/EOSShared.uplugin
+++ b/Engine/Plugins/Online/EOSShared/EOSShared.uplugin
@@ -24,5 +24,13 @@
 				"Win64"
 			]
 		}
+	],
+
+	"Plugins": [
+		{
+			"Name": "SteamShared",
+			"Enabled": false, 
+			"Optional": true
+		}
 	]
 }
diff --git a/Engine/Plugins/Online/EOSShared/Source/EOSShared/EOSShared.Build.cs b/Engine/Plugins/Online/EOSShared/Source/EOSShared/EOSShared.Build.cs
--- a/Engine/Plugins/Online/EOSShared/Source/EOSShared/EOSShared.Build.cs
+++ b/Engine/Plugins/Online/EOSShared/Source/EOSShared/EOSShared.Build.cs
@@ -42,5 +42,10 @@ public class EOSShared : ModuleRules
 		{
 			PrivateDependencyModuleNames.Add("Slate");
 		}
+
+		if (Target.Platform == UnrealTargetPlatform.Win64)
+		{
+			PrivateIncludePathModuleNames.Add("SteamShared");
+		}
 	}
 }
diff --git a/Engine/Plugins/Online/EOSShared/Source/EOSShared/Private/EOSSDKManager.h b/Engine/Plugins/Online/EOSShared/Source/EOSShared/Private/EOSSDKManager.h
--- a/Engine/Plugins/Online/EOSShared/Source/EOSShared/Private/EOSSDKManager.h
+++ b/Engine/Plugins/Online/EOSShared/Source/EOSShared/Private/EOSSDKManager.h
@@ -67,7 +67,7 @@ public:
 	virtual TSharedRef<IEOSFastTickLock> GetFastTickLock() override;
 	// End IEOSSDKManager
 
-	EOS_EResult Initialize();
+	virtual EOS_EResult Initialize();
 	void Shutdown();
 
 protected:
diff --git a/Engine/Plugins/Online/EOSShared/Source/EOSShared/Private/Windows/WindowsEOSSDKManager.cpp b/Engine/Plugins/Online/EOSShared/Source/EOSShared/Private/Windows/WindowsEOSSDKManager.cpp
--- a/Engine/Plugins/Online/EOSShared/Source/EOSShared/Private/Windows/WindowsEOSSDKManager.cpp
+++ b/Engine/Plugins/Online/EOSShared/Source/EOSShared/Private/Windows/WindowsEOSSDKManager.cpp
@@ -6,8 +6,11 @@
 
 #include "HAL/FileManager.h"
 #include "Misc/Paths.h"
+#include "Misc/ConfigCacheIni.h"
 #include "EOSShared.h"
 #include "Windows/eos_Windows.h"
+#include "Modules/ModuleManager.h"
+#include "ISteamSharedModule.h"
 
 FWindowsEOSSDKManager::FWindowsEOSSDKManager()
 {
@@ -19,6 +22,18 @@ FWindowsEOSSDKManager::FWindowsEOSSDKManager()
 	PlatformSteamOptions.SteamApiInterfaceVersionsArrayBytes = 0;
 
 	UE_EOS_CHECK_API_MISMATCH(EOS_INTEGRATEDPLATFORM_STEAM_OPTIONS_API_LATEST, 3);
+
+	GConfig->GetBool(TEXT("EOSSDK"), TEXT("bEnablePlatformIntegration"), bEnablePlatformIntegration, GEngineIni);
+
+	if (bEnablePlatformIntegration)
+	{
+		ISteamSharedModule* SteamSharedModule = FModuleManager::LoadModulePtr<ISteamSharedModule>(TEXT("SteamShared"));
+
+		if (!SteamSharedModule)
+		{
+			UE_LOG(LogEOSShared, Warning, TEXT("SteamShared module not available, creation of platforms requiring Steam integration will fail, ensure SteamShared plugin is enabled"));
+		}
+	}
 }
 
 FWindowsEOSSDKManager::~FWindowsEOSSDKManager()
@@ -26,6 +41,32 @@ FWindowsEOSSDKManager::~FWindowsEOSSDKManager()
 
 }
 
+EOS_EResult FWindowsEOSSDKManager::Initialize()
+{
+	if (bEnablePlatformIntegration)
+	{
+		// Register a shared pointer to Steam client handle. This is needed to ensure steam client handle is valid for EOS SDK to use.
+		if (ISteamSharedModule* SteamSharedModule = FModuleManager::GetModulePtr<ISteamSharedModule>(TEXT("SteamShared")))
+		{
+			SteamAPIClientHandle = SteamSharedModule->ObtainSteamClientInstanceHandle();
+		}
+	}
+
+	return Super::Initialize();
+}
+
+IEOSPlatformHandlePtr FWindowsEOSSDKManager::CreatePlatform(EOS_Platform_Options& PlatformOptions)
+{
+	// If native platform integration is enabled, and the Steam client isn't avaible, the EOS Platform should not be created
+	if (bEnablePlatformIntegration && !SteamAPIClientHandle.IsValid())
+	{
+		UE_LOG(LogEOSShared, Warning, TEXT("FWindowsEOSSDKManager::CreatePlatform failed. Steam integrated platform is enabled and steam initialization failed. EosPlatformHandle=nullptr"));
+		return nullptr;
+	}
+
+	return Super::CreatePlatform(PlatformOptions); 
+}
+
 IEOSPlatformHandlePtr FWindowsEOSSDKManager::CreatePlatform(const FEOSSDKPlatformConfig& PlatformConfig, EOS_Platform_Options& PlatformOptions)
 {
 	if (PlatformConfig.bWindowsEnableOverlayD3D9) PlatformOptions.Flags |= EOS_PF_WINDOWS_ENABLE_OVERLAY_D3D9;
diff --git a/Engine/Plugins/Online/EOSShared/Source/EOSShared/Private/Windows/WindowsEOSSDKManager.h b/Engine/Plugins/Online/EOSShared/Source/EOSShared/Private/Windows/WindowsEOSSDKManager.h
--- a/Engine/Plugins/Online/EOSShared/Source/EOSShared/Private/Windows/WindowsEOSSDKManager.h
+++ b/Engine/Plugins/Online/EOSShared/Source/EOSShared/Private/Windows/WindowsEOSSDKManager.h
@@ -13,12 +13,16 @@ public:
 
 	FWindowsEOSSDKManager();
 	virtual ~FWindowsEOSSDKManager();
+	virtual EOS_EResult Initialize() override;
+	virtual IEOSPlatformHandlePtr CreatePlatform(EOS_Platform_Options& PlatformOptions) override;
 protected:
 	virtual IEOSPlatformHandlePtr CreatePlatform(const FEOSSDKPlatformConfig& PlatformConfig, EOS_Platform_Options& PlatformOptions) override;
 
 	virtual const void* GetIntegratedPlatformOptions() override;
 	virtual EOS_IntegratedPlatformType GetIntegratedPlatformType() override;
-
+	
+	/** Steam Client API Handle */
+	TSharedPtr<class FSteamClientInstanceHandler> SteamAPIClientHandle;
 
 	EOS_IntegratedPlatform_Steam_Options PlatformSteamOptions;
 
diff --git a/Engine/Plugins/Runtime/Steam/SteamShared/Source/SteamShared/Public/ISteamSharedModule.h b/Engine/Plugins/Runtime/Steam/SteamShared/Source/SteamShared/Public/ISteamSharedModule.h
new file mode 100644
--- /dev/null
+++ b/Engine/Plugins/Runtime/Steam/SteamShared/Source/SteamShared/Public/ISteamSharedModule.h
@@ -0,0 +1,13 @@
+// Copyright Epic Games, Inc. All Rights Reserved.
+
+#pragma once
+
+#include "Modules/ModuleInterface.h"
+#include "Templates/SharedPointer.h"
+
+class ISteamSharedModule : public IModuleInterface
+{
+public:
+	/** Retrieve steam client handle if steam is enabled, otherwise return a nullptr */
+	virtual TSharedPtr<class FSteamClientInstanceHandler> ObtainSteamClientInstanceHandle() = 0;
+};
\ No newline at end of file
diff --git a/Engine/Plugins/Runtime/Steam/SteamShared/Source/SteamShared/Public/SteamSharedModule.h b/Engine/Plugins/Runtime/Steam/SteamShared/Source/SteamShared/Public/SteamSharedModule.h
--- a/Engine/Plugins/Runtime/Steam/SteamShared/Source/SteamShared/Public/SteamSharedModule.h
+++ b/Engine/Plugins/Runtime/Steam/SteamShared/Source/SteamShared/Public/SteamSharedModule.h
@@ -3,7 +3,7 @@
 #pragma once
 
 #include "CoreMinimal.h"
-#include "Modules/ModuleInterface.h"
+#include "ISteamSharedModule.h"
 #include "Modules/ModuleManager.h"
 #include "SteamSharedPackage.h"
 
@@ -11,7 +11,7 @@
 #define LOADING_STEAM_SERVER_LIBRARY_DYNAMICALLY		(PLATFORM_WINDOWS || (PLATFORM_LINUX && !IS_MONOLITHIC) || PLATFORM_MAC)
 #define LOADING_STEAM_LIBRARIES_DYNAMICALLY				(LOADING_STEAM_CLIENT_LIBRARY_DYNAMICALLY || LOADING_STEAM_SERVER_LIBRARY_DYNAMICALLY)
 
-class STEAMSHARED_API FSteamSharedModule : public IModuleInterface
+class STEAMSHARED_API FSteamSharedModule : public ISteamSharedModule
 {
 public:
 
@@ -42,7 +42,7 @@ public:
 	 *
 	 * @return A handler to the Steam Client API, use IsValid to check if the handle is initialized.
 	 */
-	TSharedPtr<class FSteamClientInstanceHandler> ObtainSteamClientInstanceHandle();
+	virtual TSharedPtr<class FSteamClientInstanceHandler> ObtainSteamClientInstanceHandle() override;
 
 	/**
 	 * Initializes Steam Server API and provides a handler that will keep the API valid for the lifetime of the
