Miscellaneous fixes for Windows Arm64

#rnx

diff --git a/Engine/Plugins/Experimental/OpenImageDenoise/Source/OpenImageDenoise/Private/Plugin.cpp b/Engine/Plugins/Experimental/OpenImageDenoise/Source/OpenImageDenoise/Private/Plugin.cpp
--- a/Engine/Plugins/Experimental/OpenImageDenoise/Source/OpenImageDenoise/Private/Plugin.cpp
+++ b/Engine/Plugins/Experimental/OpenImageDenoise/Source/OpenImageDenoise/Private/Plugin.cpp
@@ -12,6 +12,7 @@
 #include "RHITypes.h"
 #include "SceneView.h"
 
+#if WITH_INTELOIDN
 #include "OpenImageDenoise/oidn.hpp"
 
 BEGIN_SHADER_PARAMETER_STRUCT(FDenoiseTextureParameters, )
@@ -320,3 +321,5 @@ void FOpenImageDenoiseModule::ShutdownModule()
 	// Release scratch memory and destroy the OIDN device and filters
 	UnregisterDenoiser(TEXT("OIDN"));
 }
+
+#endif
diff --git a/Engine/Plugins/Online/EOSShared/Source/EOSShared/Private/Windows/WindowsEOSSDKManager.cpp b/Engine/Plugins/Online/EOSShared/Source/EOSShared/Private/Windows/WindowsEOSSDKManager.cpp
--- a/Engine/Plugins/Online/EOSShared/Source/EOSShared/Private/Windows/WindowsEOSSDKManager.cpp
+++ b/Engine/Plugins/Online/EOSShared/Source/EOSShared/Private/Windows/WindowsEOSSDKManager.cpp
@@ -5,6 +5,7 @@
 #include "WindowsEOSSDKManager.h"
 
 #include "HAL/FileManager.h"
+#include "HAL/PlatformProcess.h"
 #include "Misc/Paths.h"
 #include "Misc/ConfigCacheIni.h"
 #include "EOSShared.h"
@@ -82,6 +83,7 @@ IEOSPlatformHandlePtr FWindowsEOSSDKManager::CreatePlatform(const FEOSSDKPlatfor
 	if (PlatformConfig.bWindowsEnableOverlayD3D10) PlatformOptions.Flags |= EOS_PF_WINDOWS_ENABLE_OVERLAY_D3D10;
 	if (PlatformConfig.bWindowsEnableOverlayOpenGL) PlatformOptions.Flags |= EOS_PF_WINDOWS_ENABLE_OVERLAY_OPENGL;
 
+#if !PLATFORM_CPU_ARM_FAMILY
 	if (PlatformConfig.bEnableRTC)
 	{
 		static const FString XAudioPath = IFileManager::Get().ConvertToAbsolutePathForExternalAppForRead(*FPaths::Combine(FPaths::EngineDir(), TEXT("Binaries/ThirdParty/Windows/XAudio2_9"), PLATFORM_64BITS ? TEXT("x64") : TEXT("x86"), TEXT("xaudio2_9redist.dll")));
@@ -94,6 +96,7 @@ IEOSPlatformHandlePtr FWindowsEOSSDKManager::CreatePlatform(const FEOSSDKPlatfor
 
 		const_cast<EOS_Platform_RTCOptions*>(PlatformOptions.RTCOptions)->PlatformSpecificOptions = &WindowsRTCOptions;
 	}
+#endif
 
 	return FEOSSDKManager::CreatePlatform(PlatformConfig, PlatformOptions);
 }
diff --git a/Engine/Source/Runtime/Core/Private/MemPro/MemPro.cpp b/Engine/Source/Runtime/Core/Private/MemPro/MemPro.cpp
--- a/Engine/Source/Runtime/Core/Private/MemPro/MemPro.cpp
+++ b/Engine/Source/Runtime/Core/Private/MemPro/MemPro.cpp
@@ -4545,7 +4545,7 @@ namespace MemPro
 		#ifdef MEMPRO64
 			#ifdef _M_ARM64
 				// this will get CNTVCT_EL0, which appears to be the best replacement on arm
-				return (int)_ReadStatusReg(ARM64_SYSREG(0x3, 0x3, 0xe, 0x0, 0x2));
+				return (int)_ReadStatusReg(ARM64_CNTVCT);
 			#else
 				return (int64)__rdtsc();
 			#endif
