- Improve logging of attempts to load the EOSSDK binary

diff --git a/Engine/Plugins/Online/EOSShared/Source/EOSShared/Private/EOSSDKManager.cpp b/Engine/Plugins/Online/EOSShared/Source/EOSShared/Private/EOSSDKManager.cpp
--- a/Engine/Plugins/Online/EOSShared/Source/EOSShared/Private/EOSSDKManager.cpp
+++ b/Engine/Plugins/Online/EOSShared/Source/EOSShared/Private/EOSSDKManager.cpp
@@ -157,25 +157,32 @@ void FEOSSDKManager::LoadSDKHandle()
 {
 	check(SDKHandle == nullptr);
 
-	const FString RuntimeLibraryName = OnRequestRuntimeLibraryName.IsBound()
-		? OnRequestRuntimeLibraryName.Execute()
-		: TEXT(EOSSDK_RUNTIME_LIBRARY_NAME);
+	auto AttemptLoadDll = [](const FString& BinaryPath) -> void*
+	{
+		UE_LOG(LogEOSShared, Verbose, TEXT("Attempting to load \"%s\""), *BinaryPath);
+		void* Result = FPlatformProcess::GetDllHandle(*BinaryPath);
+		UE_CLOG(Result, LogEOSShared, Log, TEXT("Loaded \"%s\""), *BinaryPath);
+		UE_CLOG(!Result, LogEOSShared, Verbose, TEXT("Failed to load \"%s\""), *BinaryPath);
+		return Result;
+	};
 
 #if !UE_BUILD_SHIPPING
-	FString RuntimeLibraryOverride;
-	if (FParse::Value(FCommandLine::Get(), TEXT("eossdkbinary="), RuntimeLibraryOverride))
+	FString CommandLineBinary;
+	if (FParse::Value(FCommandLine::Get(), TEXT("eossdkbinary="), CommandLineBinary))
 	{
-		SDKHandle = FPlatformProcess::GetDllHandle(*RuntimeLibraryOverride);
-		UE_CLOG(SDKHandle, LogEOSShared, Log, TEXT("%hs Loaded eos binary \"%s\""), __FUNCTION__, *RuntimeLibraryOverride);
-		UE_CLOG(!SDKHandle, LogEOSShared, Error, TEXT("%hs Failed to load eos binary \"%s\""), __FUNCTION__, *RuntimeLibraryOverride);
+		SDKHandle = AttemptLoadDll(CommandLineBinary);
 	}
 	else
 #endif // !UE_BUILD_SHIPPING
 	{
+		const FString RuntimeLibraryName = OnRequestRuntimeLibraryName.IsBound()
+			? OnRequestRuntimeLibraryName.Execute()
+		: TEXT(EOSSDK_RUNTIME_LIBRARY_NAME);
+
 		const FString ProjectBinaryPath = FPaths::ConvertRelativePathToFull(FPaths::Combine(FPaths::ProjectDir(), TEXT("Binaries"), FPlatformProcess::GetBinariesSubdirectory(), *RuntimeLibraryName));
 		if (FPaths::FileExists(ProjectBinaryPath))
 		{
-			SDKHandle = FPlatformProcess::GetDllHandle(*ProjectBinaryPath);
+			SDKHandle = AttemptLoadDll(ProjectBinaryPath);
 		}
 
 		if (!SDKHandle)
@@ -183,13 +190,13 @@ void FEOSSDKManager::LoadSDKHandle()
 			const FString EngineBinaryPath = FPaths::ConvertRelativePathToFull(FPaths::Combine(FPaths::EngineDir(), TEXT("Binaries"), FPlatformProcess::GetBinariesSubdirectory(), *RuntimeLibraryName));
 			if (FPaths::FileExists(EngineBinaryPath))
 			{
-				SDKHandle = FPlatformProcess::GetDllHandle(*EngineBinaryPath);
+				SDKHandle = AttemptLoadDll(EngineBinaryPath);
 			}
 		}
 
 		if (!SDKHandle)
 		{
-			SDKHandle = FPlatformProcess::GetDllHandle(*RuntimeLibraryName);
+			SDKHandle = AttemptLoadDll(RuntimeLibraryName);
 		}
 	}
 
@@ -201,7 +208,7 @@ void FEOSSDKManager::LoadSDKHandle()
 		{
 			FPlatformMisc::MessageBoxExt(
 				EAppMsgType::Ok,
-				*FText::Format(NSLOCTEXT("EOSShared", "DllLoadFail", "Failed to load {0}. Please verify your installation. Exiting..."), FText::FromString(*RuntimeLibraryName)).ToString(),
+				*NSLOCTEXT("EOSShared", "DllLoadFail", "Failed to load EOSSDK. Please verify your installation. Exiting...").ToString(),
 				TEXT("Error")
 			);
 			UE_LOG(LogEOSShared, Fatal, TEXT("%hs Failed to load EOSSDK binary"), __FUNCTION__);
