Windows EOS SDK manager passes the system xaudio path to EOS SDK for Windows Arm64 because there is no redist version for that architecture.

- this matches the rest of the engine's expectations around xaudio.

#rnx

diff --git a/Engine/Plugins/Online/EOSShared/Source/EOSShared/Private/Windows/WindowsEOSSDKManager.cpp b/Engine/Plugins/Online/EOSShared/Source/EOSShared/Private/Windows/WindowsEOSSDKManager.cpp
--- a/Engine/Plugins/Online/EOSShared/Source/EOSShared/Private/Windows/WindowsEOSSDKManager.cpp
+++ b/Engine/Plugins/Online/EOSShared/Source/EOSShared/Private/Windows/WindowsEOSSDKManager.cpp
@@ -83,10 +83,13 @@ IEOSPlatformHandlePtr FWindowsEOSSDKManager::CreatePlatform(const FEOSSDKPlatfor
 	if (PlatformConfig.bWindowsEnableOverlayD3D10) PlatformOptions.Flags |= EOS_PF_WINDOWS_ENABLE_OVERLAY_D3D10;
 	if (PlatformConfig.bWindowsEnableOverlayOpenGL) PlatformOptions.Flags |= EOS_PF_WINDOWS_ENABLE_OVERLAY_OPENGL;
 
-#if !PLATFORM_CPU_ARM_FAMILY
 	if (PlatformConfig.bEnableRTC)
 	{
+#if !PLATFORM_CPU_ARM_FAMILY
 		static const FString XAudioPath = IFileManager::Get().ConvertToAbsolutePathForExternalAppForRead(*FPaths::Combine(FPaths::EngineDir(), TEXT("Binaries/ThirdParty/Windows/XAudio2_9"), PLATFORM_64BITS ? TEXT("x64") : TEXT("x86"), TEXT("xaudio2_9redist.dll")));
+#else
+		static const FString XAudioPath = IFileManager::Get().ConvertToAbsolutePathForExternalAppForRead(*FPaths::Combine(FWindowsPlatformProcess::WindowsSystemDir(), TEXT("xaudio2_9.dll"))); // must use the system xaudio DLL as there's no arm64 redist version
+#endif
 		static const FTCHARToUTF8 Utf8XAudioPath(*XAudioPath);
 
 		static EOS_Windows_RTCOptions WindowsRTCOptions = {};
@@ -96,7 +99,6 @@ IEOSPlatformHandlePtr FWindowsEOSSDKManager::CreatePlatform(const FEOSSDKPlatfor
 
 		const_cast<EOS_Platform_RTCOptions*>(PlatformOptions.RTCOptions)->PlatformSpecificOptions = &WindowsRTCOptions;
 	}
-#endif
 
 	return FEOSSDKManager::CreatePlatform(PlatformConfig, PlatformOptions);
 }
diff --git a/Engine/Source/Runtime/Core/Private/Windows/WindowsPlatformProcess.cpp b/Engine/Source/Runtime/Core/Private/Windows/WindowsPlatformProcess.cpp
--- a/Engine/Source/Runtime/Core/Private/Windows/WindowsPlatformProcess.cpp
+++ b/Engine/Source/Runtime/Core/Private/Windows/WindowsPlatformProcess.cpp
@@ -1403,6 +1403,24 @@ FString FWindowsPlatformProcess::GetApplicationSettingsDir(const ApplicationSett
 	return WindowsApplicationSettingsDir;
 }
 
+const TCHAR* FWindowsPlatformProcess::WindowsSystemDir()
+{
+	static FString WindowsSystemDir;
+	if( !WindowsSystemDir.Len() )
+	{
+		TCHAR* SystemPath;
+
+		HRESULT Ret = SHGetKnownFolderPath(FOLDERID_System, 0, nullptr, &SystemPath);
+		if (SUCCEEDED(Ret))
+		{
+			WindowsSystemDir = FString(SystemPath).Replace(TEXT("\\"), TEXT("/"));
+			CoTaskMemFree(SystemPath);
+		}
+	}
+	return *WindowsSystemDir;
+
+}
+
 const TCHAR* FWindowsPlatformProcess::ComputerName()
 {
 	static TCHAR Result[256] = {};
diff --git a/Engine/Source/Runtime/Core/Public/Windows/WindowsPlatformProcess.h b/Engine/Source/Runtime/Core/Public/Windows/WindowsPlatformProcess.h
--- a/Engine/Source/Runtime/Core/Public/Windows/WindowsPlatformProcess.h
+++ b/Engine/Source/Runtime/Core/Public/Windows/WindowsPlatformProcess.h
@@ -150,6 +150,7 @@ public:
 	static CORE_API const TCHAR* UserSettingsDirMediumIntegrity();
 	static CORE_API const TCHAR* ApplicationSettingsDir();
 	static CORE_API FString GetApplicationSettingsDir(const ApplicationSettingsContext& Settings);
+	static CORE_API const TCHAR* WindowsSystemDir();
 	static CORE_API const TCHAR* ComputerName();
 	static CORE_API const TCHAR* UserName(bool bOnlyAlphaNumeric = true);
 	static CORE_API void SetCurrentWorkingDirectoryToBaseDir();
