Ask for background time until all delegates called from applicationWillEnterBackground are called
Keep track of background state and audio interrupted state for sending EOS SDK to backfground

diff --git a/Engine/Plugins/Online/EOSShared/Source/EOSShared/Private/IOS/IOSEOSSDKManager.cpp b/Engine/Plugins/Online/EOSShared/Source/EOSShared/Private/IOS/IOSEOSSDKManager.cpp
--- a/Engine/Plugins/Online/EOSShared/Source/EOSShared/Private/IOS/IOSEOSSDKManager.cpp
+++ b/Engine/Plugins/Online/EOSShared/Source/EOSShared/Private/IOS/IOSEOSSDKManager.cpp
@@ -13,11 +13,13 @@ static void OnUrlOpened(UIApplication* application, NSURL* url, NSString* source
 }
 
 FIOSEOSSDKManager::FIOSEOSSDKManager()
+: bIsAppOnBackground(false)
+, bIsAudioInterrupted(false)
+
 {
-	FCoreDelegates::ApplicationHasEnteredForegroundDelegate.AddRaw(this, &FIOSEOSSDKManager::OnApplicationStatusChanged, EOS_EApplicationStatus::EOS_AS_Foreground);
-	FCoreDelegates::ApplicationWillEnterBackgroundDelegate.AddRaw(this, &FIOSEOSSDKManager::OnApplicationStatusChanged, EOS_EApplicationStatus::EOS_AS_BackgroundSuspended);
+	FCoreDelegates::ApplicationHasEnteredForegroundDelegate.AddRaw(this, &FIOSEOSSDKManager::OnApplicationForegroundChanged, false);
+	FCoreDelegates::ApplicationWillEnterBackgroundDelegate.AddRaw(this, &FIOSEOSSDKManager::OnApplicationForegroundChanged, true);
 	FCoreDelegates::AudioInterruptionDelegate.AddRaw(this, &FIOSEOSSDKManager::OnAudioInterruptedNotification);
-	FCoreDelegates::ApplicationHasReactivatedDelegate.AddRaw(this, &FIOSEOSSDKManager::OnApplicationStatusChanged, EOS_EApplicationStatus::EOS_AS_Foreground);
 	FIOSCoreDelegates::OnOpenURL.AddStatic(&OnUrlOpened);
 }
 
@@ -26,7 +28,6 @@ FIOSEOSSDKManager::~FIOSEOSSDKManager()
 	FCoreDelegates::ApplicationHasEnteredForegroundDelegate.RemoveAll(this);
 	FCoreDelegates::ApplicationWillEnterBackgroundDelegate.RemoveAll(this);
 	FCoreDelegates::AudioInterruptionDelegate.RemoveAll(this);
-	FCoreDelegates::ApplicationHasReactivatedDelegate.RemoveAll(this);
 }
 
 FString FIOSEOSSDKManager::GetCacheDirBase() const
@@ -39,11 +40,29 @@ FString FIOSEOSSDKManager::GetCacheDirBase() const
 	return FString(UTF8_TO_TCHAR(CStrCacheDirectory));
 };
 
+void FIOSEOSSDKManager::OnApplicationForegroundChanged(bool bWillBeOnBackground)
+{
+	bIsAppOnBackground = bWillBeOnBackground;
+	UpdateApplicationStatus(bWillBeOnBackground);
+}
+
 void FIOSEOSSDKManager::OnAudioInterruptedNotification(bool bInterrupted)
 {
-	if (bInterrupted)
+	bIsAudioInterrupted = bInterrupted;
+	UpdateApplicationStatus(false);
+}
+
+void FIOSEOSSDKManager::UpdateApplicationStatus(bool bShouldTick)
+{
+	EOS_EApplicationStatus NewStatus = EOS_EApplicationStatus::EOS_AS_Foreground; 
+	if (bIsAppOnBackground || bIsAudioInterrupted)
+	{
+		NewStatus = EOS_EApplicationStatus::EOS_AS_BackgroundSuspended;
+	}
+	OnApplicationStatusChanged(NewStatus);
+	if (bShouldTick)
 	{
-		OnApplicationStatusChanged(EOS_EApplicationStatus::EOS_AS_BackgroundSuspended);
+		Tick(0.f);
 	}
 }
 
diff --git a/Engine/Plugins/Online/EOSShared/Source/EOSShared/Private/IOS/IOSEOSSDKManager.h b/Engine/Plugins/Online/EOSShared/Source/EOSShared/Private/IOS/IOSEOSSDKManager.h
--- a/Engine/Plugins/Online/EOSShared/Source/EOSShared/Private/IOS/IOSEOSSDKManager.h
+++ b/Engine/Plugins/Online/EOSShared/Source/EOSShared/Private/IOS/IOSEOSSDKManager.h
@@ -14,7 +14,12 @@ public:
 
 	virtual FString GetCacheDirBase() const override;
 private:
+	void OnApplicationForegroundChanged(bool bWillBeOnBackground);
 	void OnAudioInterruptedNotification(bool bInterrupted);
+	void UpdateApplicationStatus(bool bShouldTick);
+
+	bool bIsAppOnBackground;
+	bool bIsAudioInterrupted;
 };
 
 using FPlatformEOSSDKManager = FIOSEOSSDKManager;
diff --git a/Engine/Source/Runtime/ApplicationCore/Private/IOS/IOSAppDelegate.cpp b/Engine/Source/Runtime/ApplicationCore/Private/IOS/IOSAppDelegate.cpp
--- a/Engine/Source/Runtime/ApplicationCore/Private/IOS/IOSAppDelegate.cpp
+++ b/Engine/Source/Runtime/ApplicationCore/Private/IOS/IOSAppDelegate.cpp
@@ -1626,6 +1626,12 @@ FCriticalSection RenderSuspend;
 
 	FEmbeddedCommunication::KeepAwake(TEXT("Background"), false);
 
+	__block UIBackgroundTaskIdentifier BackgroundTaskId = UIBackgroundTaskInvalid;	
+	BackgroundTaskId = [[UIApplication sharedApplication] beginBackgroundTaskWithExpirationHandler:^{
+		UE_LOG(LogIOS, Warning, TEXT("Background task expired before completing FCoreDelegates::ApplicationWillEnterBackgroundDelegate broadcast"));
+        [[UIApplication sharedApplication] endBackgroundTask: BackgroundTaskId];
+    }];
+	
 	[FIOSAsyncTask CreateTaskWithBlock : ^ bool(void)
 	{
 		// the audio context should resume immediately after interrupt, if suspended
@@ -1634,6 +1640,7 @@ FCriticalSection RenderSuspend;
 		FCoreDelegates::ApplicationWillEnterBackgroundDelegate.Broadcast();
 	
 		FEmbeddedCommunication::AllowSleep(TEXT("Background"));
+        [[UIApplication sharedApplication] endBackgroundTask: BackgroundTaskId];
 		return true;
 	}];
 }
