Work around EOS_Initialize in EOSShared modifying game thread priority on Windows

We can remove this hack when EOSSDK is fixed internally.

diff --git a/Engine/Plugins/Online/EOSShared/Source/EOSShared/Private/EOSSDKManager.cpp b/Engine/Plugins/Online/EOSShared/Source/EOSShared/Private/EOSSDKManager.cpp
--- a/Engine/Plugins/Online/EOSShared/Source/EOSShared/Private/EOSSDKManager.cpp
+++ b/Engine/Plugins/Online/EOSShared/Source/EOSShared/Private/EOSSDKManager.cpp
@@ -1174,6 +1174,12 @@ void FEOSSDKManager::Shutdown()
 	}
 }
 
+// TODO: Remove Windows workaround when EOS stops modifying thread priority
+#if PLATFORM_WINDOWS
+extern "C" int __stdcall GetThreadPriority(void* hThread);
+extern "C" void* __stdcall GetCurrentThread(void);
+extern "C" int __stdcall SetThreadPriority(void* hThread, int nPriority);
+#endif
 EOS_EResult FEOSSDKManager::EOSInitialize(EOS_InitializeOptions& Options)
 {
 	OnPreInitializeSDK.Broadcast(Options);
@@ -1181,7 +1187,18 @@ EOS_EResult FEOSSDKManager::EOSInitialize(EOS_InitializeOptions& Options)
 	EOS_InitializeOptions* OptionsPtr = &Options;
 	OnPreInitializeSDK2.Broadcast(OptionsPtr);
 
-	return EOS_Initialize(OptionsPtr);
+#if PLATFORM_WINDOWS
+	void* const hThread = ::GetCurrentThread();
+	const int Priority = ::GetThreadPriority(hThread);
+#endif
+	const EOS_EResult Result = EOS_Initialize(OptionsPtr);
+#if PLATFORM_WINDOWS
+	if (Priority != 0x7fffffff)
+	{
+		::SetThreadPriority(hThread, Priority);
+	}
+#endif
+	return Result;
 }
 
 bool FEOSSDKManager::Exec_Runtime(class UWorld* InWorld, const TCHAR* Cmd, FOutputDevice& Ar)
