Add CVar to avoid notifying EOS SDK about audio interruptions

diff --git a/Engine/Plugins/Online/EOSShared/Source/EOSShared/Private/IOS/IOSEOSSDKManager.cpp b/Engine/Plugins/Online/EOSShared/Source/EOSShared/Private/IOS/IOSEOSSDKManager.cpp
--- a/Engine/Plugins/Online/EOSShared/Source/EOSShared/Private/IOS/IOSEOSSDKManager.cpp
+++ b/Engine/Plugins/Online/EOSShared/Source/EOSShared/Private/IOS/IOSEOSSDKManager.cpp
@@ -6,6 +6,12 @@
 #include "IOSAppDelegate.h"
 #include "Misc/CoreDelegates.h"
 
+static TAutoConsoleVariable<bool> CVarNotifySDKOnAudioInterruption(
+	TEXT("IOSEOSSDK.NotifySDKOnAudioInterruption"),
+	true,
+	TEXT("Whether to notify EOSSDK on audio interruptions")
+);
+
 static void OnUrlOpened(UIApplication* application, NSURL* url, NSString* sourceApplication, id annotation)
 {
 	// TODO: This is based on a prototype fix on EOS SDK. Once the fix is properly submitted to EOS SDK we should update it
@@ -19,7 +25,10 @@ FIOSEOSSDKManager::FIOSEOSSDKManager()
 {
 	FCoreDelegates::ApplicationHasEnteredForegroundDelegate.AddRaw(this, &FIOSEOSSDKManager::OnApplicationForegroundChanged, false);
 	FCoreDelegates::ApplicationWillEnterBackgroundDelegate.AddRaw(this, &FIOSEOSSDKManager::OnApplicationForegroundChanged, true);
-	FCoreDelegates::AudioInterruptionDelegate.AddRaw(this, &FIOSEOSSDKManager::OnAudioInterruptedNotification);
+	if (CVarNotifySDKOnAudioInterruption.GetValueOnAnyThread())
+	{
+		FCoreDelegates::AudioInterruptionDelegate.AddRaw(this, &FIOSEOSSDKManager::OnAudioInterruptedNotification);
+	}
 	FIOSCoreDelegates::OnOpenURL.AddStatic(&OnUrlOpened);
 }
 
