- Fix sending invalid characters to EOSSDK in ProductName field

#rnx
#tests local, with ? character

diff --git a/Engine/Plugins/Online/EOSShared/Source/EOSShared/Private/EOSSDKManager.cpp b/Engine/Plugins/Online/EOSShared/Source/EOSShared/Private/EOSSDKManager.cpp
--- a/Engine/Plugins/Online/EOSShared/Source/EOSShared/Private/EOSSDKManager.cpp
+++ b/Engine/Plugins/Online/EOSShared/Source/EOSShared/Private/EOSSDKManager.cpp
@@ -16,6 +16,7 @@
 
 #include "Misc/ScopeRWLock.h"
 #include "Misc/App.h"
+#include "Misc/AsciiSet.h"
 #include "Misc/CoreMisc.h"
 #include "Misc/CoreDelegates.h"
 #include "Misc/Fork.h"
@@ -260,7 +261,22 @@ EOS_EResult FEOSSDKManager::Initialize()
 	{
 		UE_LOG(LogEOSShared, Log, TEXT("%hs Initializing EOSSDK Version:%s"), __FUNCTION__, UTF8_TO_TCHAR(EOS_GetVersion()));
 
-		const FTCHARToUTF8 ProductName(*GetProductName());
+		FString ProductName = GetProductName();
+		constexpr FAsciiSet ValidProductNameChars("ABCDEFGHIJKLMNOPQRSTUVWXYZabcdefghijklmnopqrstuvwxyz0123456789._ !?&()+-:");
+		if (!FAsciiSet::HasOnly(*ProductName, ValidProductNameChars))
+		{
+			UE_LOG(LogEOSShared, Warning, TEXT("ProductName=%s contains invalid characters, please set a valid one"), *ProductName);
+
+			while (TCHAR* BadChar = const_cast<TCHAR*>(FAsciiSet::Skip(*ProductName, ValidProductNameChars)))
+			{
+				if (*BadChar == TCHAR('\0'))
+				{
+					break;
+				}
+				*BadChar = TCHAR(' ');
+			}
+		}
+		const FTCHARToUTF8 ProductNameUtf8(*ProductName);
 		const FTCHARToUTF8 ProductVersion(*GetProductVersion());
 
 		EOS_InitializeOptions InitializeOptions = {};
@@ -269,7 +285,7 @@ EOS_EResult FEOSSDKManager::Initialize()
 		InitializeOptions.AllocateMemoryFunction = &EosMalloc;
 		InitializeOptions.ReallocateMemoryFunction = &EosRealloc;
 		InitializeOptions.ReleaseMemoryFunction = &EosFree;
-		InitializeOptions.ProductName = ProductName.Get();
+		InitializeOptions.ProductName = ProductNameUtf8.Get();
 		InitializeOptions.ProductVersion = ProductVersion.Length() > 0 ? ProductVersion.Get() : nullptr;
 		InitializeOptions.Reserved = nullptr;
 		InitializeOptions.SystemInitializeOptions = nullptr;
