Stop notifying EOS SDK about audio interruption events on iOS by suspending it

diff --git a/Engine/Plugins/Online/EOSShared/Source/EOSShared/Private/IOS/IOSEOSSDKManager.cpp b/Engine/Plugins/Online/EOSShared/Source/EOSShared/Private/IOS/IOSEOSSDKManager.cpp
--- a/Engine/Plugins/Online/EOSShared/Source/EOSShared/Private/IOS/IOSEOSSDKManager.cpp
+++ b/Engine/Plugins/Online/EOSShared/Source/EOSShared/Private/IOS/IOSEOSSDKManager.cpp
@@ -6,12 +6,6 @@
 #include "IOSAppDelegate.h"
 #include "Misc/CoreDelegates.h"
 
-static TAutoConsoleVariable<bool> CVarNotifySDKOnAudioInterruption(
-	TEXT("IOSEOSSDK.NotifySDKOnAudioInterruption"),
-	true,
-	TEXT("Whether to notify EOSSDK on audio interruptions")
-);
-
 static void OnUrlOpened(UIApplication* application, NSURL* url, NSString* sourceApplication, id annotation)
 {
 	// TODO: This is based on a prototype fix on EOS SDK. Once the fix is properly submitted to EOS SDK we should update it
@@ -19,16 +13,9 @@ static void OnUrlOpened(UIApplication* application, NSURL* url, NSString* source
 }
 
 FIOSEOSSDKManager::FIOSEOSSDKManager()
-: bIsAppOnBackground(false)
-, bIsAudioInterrupted(false)
-
 {
 	FCoreDelegates::ApplicationHasEnteredForegroundDelegate.AddRaw(this, &FIOSEOSSDKManager::OnApplicationForegroundChanged, false);
 	FCoreDelegates::ApplicationWillEnterBackgroundDelegate.AddRaw(this, &FIOSEOSSDKManager::OnApplicationForegroundChanged, true);
-	if (CVarNotifySDKOnAudioInterruption.GetValueOnAnyThread())
-	{
-		FCoreDelegates::AudioInterruptionDelegate.AddRaw(this, &FIOSEOSSDKManager::OnAudioInterruptedNotification);
-	}
 	FIOSCoreDelegates::OnOpenURL.AddStatic(&OnUrlOpened);
 }
 
@@ -36,7 +23,6 @@ FIOSEOSSDKManager::~FIOSEOSSDKManager()
 {	
 	FCoreDelegates::ApplicationHasEnteredForegroundDelegate.RemoveAll(this);
 	FCoreDelegates::ApplicationWillEnterBackgroundDelegate.RemoveAll(this);
-	FCoreDelegates::AudioInterruptionDelegate.RemoveAll(this);
 }
 
 FString FIOSEOSSDKManager::GetCacheDirBase() const
@@ -51,25 +37,9 @@ FString FIOSEOSSDKManager::GetCacheDirBase() const
 
 void FIOSEOSSDKManager::OnApplicationForegroundChanged(bool bWillBeOnBackground)
 {
-	bIsAppOnBackground = bWillBeOnBackground;
-	UpdateApplicationStatus(bWillBeOnBackground);
-}
-
-void FIOSEOSSDKManager::OnAudioInterruptedNotification(bool bInterrupted)
-{
-	bIsAudioInterrupted = bInterrupted;
-	UpdateApplicationStatus(false);
-}
-
-void FIOSEOSSDKManager::UpdateApplicationStatus(bool bShouldTick)
-{
-	EOS_EApplicationStatus NewStatus = EOS_EApplicationStatus::EOS_AS_Foreground; 
-	if (bIsAppOnBackground || bIsAudioInterrupted)
-	{
-		NewStatus = EOS_EApplicationStatus::EOS_AS_BackgroundSuspended;
-	}
+	EOS_EApplicationStatus NewStatus = bWillBeOnBackground ? EOS_EApplicationStatus::EOS_AS_BackgroundSuspended : EOS_EApplicationStatus::EOS_AS_Foreground;
 	OnApplicationStatusChanged(NewStatus);
-	if (bShouldTick)
+	if (bWillBeOnBackground)
 	{
 		Tick(0.f);
 	}
diff --git a/Engine/Plugins/Online/EOSShared/Source/EOSShared/Private/IOS/IOSEOSSDKManager.h b/Engine/Plugins/Online/EOSShared/Source/EOSShared/Private/IOS/IOSEOSSDKManager.h
--- a/Engine/Plugins/Online/EOSShared/Source/EOSShared/Private/IOS/IOSEOSSDKManager.h
+++ b/Engine/Plugins/Online/EOSShared/Source/EOSShared/Private/IOS/IOSEOSSDKManager.h
@@ -15,11 +15,6 @@ public:
 	virtual FString GetCacheDirBase() const override;
 private:
 	void OnApplicationForegroundChanged(bool bWillBeOnBackground);
-	void OnAudioInterruptedNotification(bool bInterrupted);
-	void UpdateApplicationStatus(bool bShouldTick);
-
-	bool bIsAppOnBackground;
-	bool bIsAudioInterrupted;
 };
 
 using FPlatformEOSSDKManager = FIOSEOSSDKManager;
