- Add mechanism for suppressing log messages/categories to VeryVerbose

#rnx

diff --git a/Engine/Plugins/Online/EOSShared/Source/EOSShared/Private/EOSSDKManager.cpp b/Engine/Plugins/Online/EOSShared/Source/EOSShared/Private/EOSSDKManager.cpp
--- a/Engine/Plugins/Online/EOSShared/Source/EOSShared/Private/EOSSDKManager.cpp
+++ b/Engine/Plugins/Online/EOSShared/Source/EOSShared/Private/EOSSDKManager.cpp
@@ -83,46 +83,41 @@ namespace
 	}
 
 #if !NO_LOGGING
-	bool IsLogLevelSuppressable(EOS_ELogLevel Level)
-	{
-		return Level > EOS_ELogLevel::EOS_LOG_Off && Level < EOS_ELogLevel::EOS_LOG_Info;
-	}
-
 	void EOS_CALL EOSLogMessageReceived(const EOS_LogMessage* Message)
 	{
-#define EOSLOG(Level) UE_LOG(LogEOSSDK, Level, TEXT("%s: %s"), UTF8_TO_TCHAR(Message->Category), *MessageStr)
-#define EOSLOG_SUPPRESS(Level) if(bSuppressLogLevel) { EOSLOG(Log); } else { EOSLOG(Level); }
-
 		FString MessageStr(UTF8_TO_TCHAR(Message->Message));
 		MessageStr.TrimStartAndEndInline();
-
-		// Check if this log line is suppressed.
-		bool bSuppressLogLevel = false;
-		if (IsLogLevelSuppressable(Message->Level))
+		FString CategoryStr(UTF8_TO_TCHAR(Message->Category));
+		CategoryStr.TrimStartAndEndInline();
+		
+		EOS_ELogLevel EOSLogLevel = Message->Level;
+			
+		if (FEOSSharedModule* Module = FEOSSharedModule::Get())
 		{
-			if(FEOSSharedModule* Module = FEOSSharedModule::Get())
+			auto MessageProjection = [&MessageStr](const FString& Elem) { return MessageStr.Contains(Elem); };
+			auto CategoryProjection = [&CategoryStr](const FString& Elem) { return CategoryStr.Contains(Elem); };
+
+			if (EOSLogLevel < EOS_ELogLevel::EOS_LOG_VeryVerbose
+				&& (Algo::AnyOf(Module->GetSuppressedLogStrings_VeryVerbose(), MessageProjection)
+				|| Algo::AnyOf(Module->GetSuppressedLogCategories_VeryVerbose(), CategoryProjection)))
 			{
-				bSuppressLogLevel = Algo::AnyOf(Module->GetSuppressedLogStrings(), [&MessageStr](const FString& SuppressedLogString) { return MessageStr.Contains(SuppressedLogString); });
+				EOSLogLevel = EOS_ELogLevel::EOS_LOG_VeryVerbose;
+			}
 
-				// Check for suppressed categories if not already suppressed
-				if (!bSuppressLogLevel)
-				{
-					const TArray<FString>& SuppressedCategories = Module->GetSuppressedLogCategories();
-					if (SuppressedCategories.Num())
-					{
-						FString CategoryStr(UTF8_TO_TCHAR(Message->Category));
-						CategoryStr.TrimStartAndEndInline();
-						bSuppressLogLevel = Algo::AnyOf(SuppressedCategories, [&CategoryStr](const FString& SuppressedLogCategory) { return CategoryStr.Contains(SuppressedLogCategory); });
-					}
-				}
+			if (EOSLogLevel < EOS_ELogLevel::EOS_LOG_Info
+				&& (Algo::AnyOf(Module->GetSuppressedLogStrings_Log(), MessageProjection)
+				|| Algo::AnyOf(Module->GetSuppressedLogCategories_Log(), CategoryProjection)))
+			{
+				EOSLogLevel = EOS_ELogLevel::EOS_LOG_Info;
 			}
 		}
 
-		switch (Message->Level)
+		#define EOSLOG(Level) UE_LOG(LogEOSSDK, Level, TEXT("%s: %s"), *CategoryStr, *MessageStr)
+		switch (EOSLogLevel)
 		{
-		case EOS_ELogLevel::EOS_LOG_Fatal:			EOSLOG_SUPPRESS(Fatal); break;
-		case EOS_ELogLevel::EOS_LOG_Error:			EOSLOG_SUPPRESS(Error); break;
-		case EOS_ELogLevel::EOS_LOG_Warning:		EOSLOG_SUPPRESS(Warning); break;
+		case EOS_ELogLevel::EOS_LOG_Fatal:			EOSLOG(Fatal); break;
+		case EOS_ELogLevel::EOS_LOG_Error:			EOSLOG(Error); break;
+		case EOS_ELogLevel::EOS_LOG_Warning:		EOSLOG(Warning); break;
 		case EOS_ELogLevel::EOS_LOG_Info:			EOSLOG(Log); break;
 		case EOS_ELogLevel::EOS_LOG_Verbose:		EOSLOG(Verbose); break;
 		case EOS_ELogLevel::EOS_LOG_VeryVerbose:	EOSLOG(VeryVerbose); break;
@@ -131,8 +126,7 @@ namespace
 			// do nothing
 			break;
 		}
-#undef EOSLOG
-#undef EOSLOG_SUPPRESS
+		#undef EOSLOG
 	}
 
 	EOS_ELogLevel ConvertLogLevel(ELogVerbosity::Type UELogLevel)
diff --git a/Engine/Plugins/Online/EOSShared/Source/EOSShared/Private/EOSSharedModule.cpp b/Engine/Plugins/Online/EOSShared/Source/EOSShared/Private/EOSSharedModule.cpp
--- a/Engine/Plugins/Online/EOSShared/Source/EOSShared/Private/EOSSharedModule.cpp
+++ b/Engine/Plugins/Online/EOSShared/Source/EOSShared/Private/EOSSharedModule.cpp
@@ -108,8 +108,10 @@ void FEOSSharedModule::OnConfigSectionsChanged(const FString& IniFilename, const
 
 void FEOSSharedModule::LoadConfig()
 {
-	GConfig->GetArray(CONFIG_SECTION_NAME, TEXT("SuppressedLogStrings"), SuppressedLogStrings, GEngineIni);
-	GConfig->GetArray(CONFIG_SECTION_NAME, TEXT("SuppressedLogCategories"), SuppressedLogCategories, GEngineIni);
+	GConfig->GetArray(CONFIG_SECTION_NAME, TEXT("SuppressedLogStrings"), SuppressedLogStrings_Log, GEngineIni);
+	GConfig->GetArray(CONFIG_SECTION_NAME, TEXT("SuppressedLogCategories"), SuppressedLogCategories_Log, GEngineIni);
+	GConfig->GetArray(CONFIG_SECTION_NAME, TEXT("SuppressedLogStrings_VeryVerbose"), SuppressedLogStrings_VeryVerbose, GEngineIni);
+	GConfig->GetArray(CONFIG_SECTION_NAME, TEXT("SuppressedLogCategories_VeryVerbose"), SuppressedLogCategories_VeryVerbose, GEngineIni);
 }
 
 #undef CONFIG_SECTION_NAME
diff --git a/Engine/Plugins/Online/EOSShared/Source/EOSShared/Private/EOSSharedModule.h b/Engine/Plugins/Online/EOSShared/Source/EOSShared/Private/EOSSharedModule.h
--- a/Engine/Plugins/Online/EOSShared/Source/EOSShared/Private/EOSSharedModule.h
+++ b/Engine/Plugins/Online/EOSShared/Source/EOSShared/Private/EOSSharedModule.h
@@ -18,9 +18,10 @@ public:
 
 	static FEOSSharedModule* Get();
 
-	const TArray<FString>& GetSuppressedLogStrings() const { return SuppressedLogStrings; }
-	const TArray<FString>& GetSuppressedLogCategories() const { return SuppressedLogCategories; }
-
+	const TArray<FString>& GetSuppressedLogStrings_Log() const { return SuppressedLogStrings_Log; }
+	const TArray<FString>& GetSuppressedLogCategories_Log() const { return SuppressedLogCategories_Log; }
+	const TArray<FString>& GetSuppressedLogStrings_VeryVerbose() const { return SuppressedLogStrings_VeryVerbose; }
+	const TArray<FString>& GetSuppressedLogCategories_VeryVerbose() const { return SuppressedLogCategories_VeryVerbose; }
 private:
 	// ~Begin IModuleInterface
 	virtual void StartupModule() override;
@@ -34,7 +35,9 @@ private:
 #if WITH_EOS_SDK
 	TUniquePtr<FEOSSDKManager> SDKManager;
 #endif
-	TArray<FString> SuppressedLogStrings;
-	TArray<FString> SuppressedLogCategories;
+	TArray<FString> SuppressedLogStrings_Log;
+	TArray<FString> SuppressedLogCategories_Log;
+	TArray<FString> SuppressedLogStrings_VeryVerbose;
+	TArray<FString> SuppressedLogCategories_VeryVerbose;
 	FDelegateHandle OnPostForkDelegateHandle;
 };
\ No newline at end of file
