Analytics Provider work

diff --git a/Engine/Plugins/Online/EOSShared/Source/EOSShared/Private/EOSSDKManager.cpp b/Engine/Plugins/Online/EOSShared/Source/EOSShared/Private/EOSSDKManager.cpp
--- a/Engine/Plugins/Online/EOSShared/Source/EOSShared/Private/EOSSDKManager.cpp
+++ b/Engine/Plugins/Online/EOSShared/Source/EOSShared/Private/EOSSDKManager.cpp
@@ -1028,6 +1028,11 @@ void FEOSSDKManager::OnLogVerbosityChanged(const FLogCategoryName& CategoryName,
 #endif // !NO_LOGGING
 }
 
+FString FEOSSDKManager::GetSDKVersion() const
+{
+	return FString((const UTF8CHAR*)EOS_GetVersion());
+}
+
 FString FEOSSDKManager::GetProductName() const
 {
 	FString ProductName;
diff --git a/Engine/Plugins/Online/EOSShared/Source/EOSShared/Private/EOSSDKManager.h b/Engine/Plugins/Online/EOSShared/Source/EOSShared/Private/EOSSDKManager.h
--- a/Engine/Plugins/Online/EOSShared/Source/EOSShared/Private/EOSSDKManager.h
+++ b/Engine/Plugins/Online/EOSShared/Source/EOSShared/Private/EOSSDKManager.h
@@ -48,6 +48,7 @@ public:
 	virtual IEOSPlatformHandlePtr CreatePlatform(EOS_Platform_Options& PlatformOptions) override;
 	virtual TArray<IEOSPlatformHandlePtr> GetActivePlatforms() override;
 
+	virtual FString GetSDKVersion() const override;
 	virtual FString GetProductName() const override;
 	virtual FString GetProductVersion() const override;
 	virtual FString GetCacheDirBase() const override;
diff --git a/Engine/Plugins/Online/EOSShared/Source/EOSShared/Public/IEOSSDKManager.h b/Engine/Plugins/Online/EOSShared/Source/EOSShared/Public/IEOSSDKManager.h
--- a/Engine/Plugins/Online/EOSShared/Source/EOSShared/Public/IEOSSDKManager.h
+++ b/Engine/Plugins/Online/EOSShared/Source/EOSShared/Public/IEOSSDKManager.h
@@ -133,6 +133,7 @@ public:
 	/** Retrieves the array of platform handles for all active platforms */
 	virtual TArray<IEOSPlatformHandlePtr> GetActivePlatforms() = 0;
 
+	virtual FString GetSDKVersion() const = 0;
 	virtual FString GetProductName() const = 0;
 	virtual FString GetProductVersion() const = 0;
 	virtual FString GetCacheDirBase() const = 0;
